deprecated FTicker and family and replaced by thread-safe FTSTicker

diff --git a/Engine/Plugins/Runtime/ChunkDownloader/Source/Public/ChunkDownloader.h b/Engine/Plugins/Runtime/ChunkDownloader/Source/Public/ChunkDownloader.h
--- a/Engine/Plugins/Runtime/ChunkDownloader/Source/Public/ChunkDownloader.h
+++ b/Engine/Plugins/Runtime/ChunkDownloader/Source/Public/ChunkDownloader.h
@@ -3,6 +3,7 @@
 #pragma once
 
 #include "CoreMinimal.h"
+#include "Containers/Ticker.h"
 
 template<typename TTask> class FAsyncTask;
 class IHttpRequest;
@@ -258,7 +259,7 @@ private:
 	bool bNeedsManifestSave = false;
 
 	// handle for the per-frame mount ticker in the main thread
-	FDelegateHandle MountTicker;
+	FTSTicker::FDelegateHandle MountTicker;
 
 	// manifest download request
 	TSharedPtr<IHttpRequest, ESPMode::ThreadSafe> ManifestRequest;
diff --git a/Engine/Plugins/Runtime/HTTPChunkInstaller/Source/Public/HTTPChunkInstaller.h b/Engine/Plugins/Runtime/HTTPChunkInstaller/Source/Public/HTTPChunkInstaller.h
--- a/Engine/Plugins/Runtime/HTTPChunkInstaller/Source/Public/HTTPChunkInstaller.h
+++ b/Engine/Plugins/Runtime/HTTPChunkInstaller/Source/Public/HTTPChunkInstaller.h
@@ -12,7 +12,7 @@
 /**
 * HTTP based implementation of chunk based install
 **/
-class HTTPCHUNKINSTALLER_API FHTTPChunkInstall : public FGenericPlatformChunkInstall, public FTickerObjectBase
+class HTTPCHUNKINSTALLER_API FHTTPChunkInstall : public FGenericPlatformChunkInstall, public FTSTickerObjectBase
 {
 public:
 	FHTTPChunkInstall();
diff --git a/Engine/Plugins/Runtime/NetworkPredictionInsights/Source/NetworkPredictionInsights/Private/NetworkPredictionInsightsModule.cpp b/Engine/Plugins/Runtime/NetworkPredictionInsights/Source/NetworkPredictionInsights/Private/NetworkPredictionInsightsModule.cpp
--- a/Engine/Plugins/Runtime/NetworkPredictionInsights/Source/NetworkPredictionInsights/Private/NetworkPredictionInsightsModule.cpp
+++ b/Engine/Plugins/Runtime/NetworkPredictionInsights/Source/NetworkPredictionInsights/Private/NetworkPredictionInsightsModule.cpp
@@ -43,7 +43,7 @@ void FNetworkPredictionInsightsModule::StartupModule()
 	// (only SetUnrealInsightsLayoutIni which is extending the layouts of pre made individual tabs, not the the overall session layout)
 	if (!GIsEditor)
 	{
-		TickerHandle = FTicker::GetCoreTicker().AddTicker(TEXT("NetworkPredictionInsights"), 0.0f, [&UnrealInsightsModule](float DeltaTime)
+		TickerHandle = FTSTicker::GetCoreTicker().AddTicker(TEXT("NetworkPredictionInsights"), 0.0f, [&UnrealInsightsModule](float DeltaTime)
 		{
 			QUICK_SCOPE_CYCLE_COUNTER(STAT_FNetworkPredictionInsightsModule_Tick);
 			auto SessionPtr = UnrealInsightsModule.GetAnalysisSession();
@@ -136,7 +136,7 @@ void FNetworkPredictionInsightsModule::ShutdownModule()
 		FCoreDelegates::OnFEngineLoopInitComplete.Remove(StoreServiceHandle);
 	}
 
-	FTicker::GetCoreTicker().RemoveTicker(TickerHandle);
+	FTSTicker::GetCoreTicker().RemoveTicker(TickerHandle);
 
 	IModularFeatures::Get().UnregisterModularFeature(TraceServices::ModuleFeatureName, &NetworkPredictionTraceModule);
 }
diff --git a/Engine/Plugins/Runtime/Oculus/OculusAudio/Source/OculusAudio/Private/OculusAudioDllManager.cpp b/Engine/Plugins/Runtime/Oculus/OculusAudio/Source/OculusAudio/Private/OculusAudioDllManager.cpp
--- a/Engine/Plugins/Runtime/Oculus/OculusAudio/Source/OculusAudio/Private/OculusAudioDllManager.cpp
+++ b/Engine/Plugins/Runtime/Oculus/OculusAudio/Source/OculusAudio/Private/OculusAudioDllManager.cpp
@@ -32,7 +32,7 @@ FOculusAudioLibraryManager::FOculusAudioLibraryManager()
 
 FOculusAudioLibraryManager::~FOculusAudioLibraryManager()
 {
-	FTicker::GetCoreTicker().RemoveTicker(TickDelegateHandle);
+	FTSTicker::GetCoreTicker().RemoveTicker(TickDelegateHandle);
 }
 
 void FOculusAudioLibraryManager::Initialize()
@@ -212,7 +212,7 @@ ovrAudioContext FOculusAudioLibraryManager::GetPluginContext()
 
 			// Tick the scene from here since there is no listener
 			auto TickDelegate = FTickerDelegate::CreateRaw(this, &FOculusAudioLibraryManager::UpdatePluginContext);
-			TickDelegateHandle = FTicker::GetCoreTicker().AddTicker(TickDelegate);
+			TickDelegateHandle = FTSTicker::GetCoreTicker().AddTicker(TickDelegate);
 		}
 	}
 	return CachedPluginContext;
diff --git a/Engine/Plugins/Runtime/Oculus/OculusAudio/Source/OculusAudio/Private/OculusAudioMixer.cpp b/Engine/Plugins/Runtime/Oculus/OculusAudio/Source/OculusAudio/Private/OculusAudioMixer.cpp
--- a/Engine/Plugins/Runtime/Oculus/OculusAudio/Source/OculusAudio/Private/OculusAudioMixer.cpp
+++ b/Engine/Plugins/Runtime/Oculus/OculusAudio/Source/OculusAudio/Private/OculusAudioMixer.cpp
@@ -48,7 +48,7 @@ void OculusAudioSpatializationAudioMixer::Initialize(const FAudioPluginInitializ
 	const UOculusAudioSettings* Settings = GetDefault<UOculusAudioSettings>();
 	ApplyOculusAudioSettings(Settings);
 
-	TickDelegateHandle = FTicker::GetCoreTicker().AddTicker(FTickerDelegate::CreateRaw(this, &OculusAudioSpatializationAudioMixer::Tick));
+	TickDelegateHandle = FTSTicker::GetCoreTicker().AddTicker(FTickerDelegate::CreateRaw(this, &OculusAudioSpatializationAudioMixer::Tick));
 }
 
 void OculusAudioSpatializationAudioMixer::ApplyOculusAudioSettings(const UOculusAudioSettings* Settings)
@@ -83,7 +83,7 @@ void OculusAudioSpatializationAudioMixer::Shutdown()
 {
 	if (TickDelegateHandle.IsValid())
 	{
-		FTicker::GetCoreTicker().RemoveTicker(TickDelegateHandle);
+		FTSTicker::GetCoreTicker().RemoveTicker(TickDelegateHandle);
 	}
 }
 
diff --git a/Engine/Plugins/Runtime/Oculus/OculusAudio/Source/OculusAudio/Private/OculusAudioMixer.h b/Engine/Plugins/Runtime/Oculus/OculusAudio/Source/OculusAudio/Private/OculusAudioMixer.h
--- a/Engine/Plugins/Runtime/Oculus/OculusAudio/Source/OculusAudio/Private/OculusAudioMixer.h
+++ b/Engine/Plugins/Runtime/Oculus/OculusAudio/Source/OculusAudio/Private/OculusAudioMixer.h
@@ -7,6 +7,7 @@
 #include "OculusAudioSettings.h"
 #include "Containers/Ticker.h"
 #include "AudioMixerDevice.h"
+#include "Containers/Ticker.h"
 
 /************************************************************************/
 /* OculusAudioSpatializationAudioMixer									*/
@@ -60,5 +61,5 @@ private:
 	FCriticalSection ContextLock;
 
 	FAudioPluginInitializationParams InitParams;
-	FDelegateHandle TickDelegateHandle;
+	FTSTicker::FDelegateHandle TickDelegateHandle;
 };
\ No newline at end of file
diff --git a/Engine/Plugins/Runtime/Oculus/OculusVR/Source/OculusHMD/Private/OculusHMD_Splash.cpp b/Engine/Plugins/Runtime/Oculus/OculusVR/Source/OculusHMD/Private/OculusHMD_Splash.cpp
--- a/Engine/Plugins/Runtime/Oculus/OculusVR/Source/OculusHMD/Private/OculusHMD_Splash.cpp
+++ b/Engine/Plugins/Runtime/Oculus/OculusVR/Source/OculusHMD/Private/OculusHMD_Splash.cpp
@@ -214,7 +214,7 @@ void FSplash::StartTicker()
 
 	if (!Ticker.IsValid())
 	{
-		Ticker = MakeShareable(new FTicker(this));
+		Ticker = MakeShareable(new FTSTicker(this));
 
 		ExecuteOnRenderThread([this]()
 		{
diff --git a/Engine/Plugins/Runtime/Oculus/OculusVR/Source/OculusHMD/Private/OculusHMD_Splash.h b/Engine/Plugins/Runtime/Oculus/OculusVR/Source/OculusHMD/Private/OculusHMD_Splash.h
--- a/Engine/Plugins/Runtime/Oculus/OculusVR/Source/OculusHMD/Private/OculusHMD_Splash.h
+++ b/Engine/Plugins/Runtime/Oculus/OculusVR/Source/OculusHMD/Private/OculusHMD_Splash.h
@@ -41,10 +41,10 @@ public:
 class FSplash : public IXRLoadingScreen, public TSharedFromThis<FSplash>
 {
 protected:
-	class FTicker : public FTickableObjectRenderThread, public TSharedFromThis<FTicker>
+	class FTSTicker : public FTickableObjectRenderThread, public TSharedFromThis<FTSTicker>
 	{
 	public:
-		FTicker(FSplash* InSplash) : FTickableObjectRenderThread(false, true), pSplash(InSplash) {}
+		FTSTicker(FSplash* InSplash) : FTickableObjectRenderThread(false, true), pSplash(InSplash) {}
 
 		virtual void Tick(float DeltaTime) override { pSplash->Tick_RenderThread(DeltaTime); }
 		virtual TStatId GetStatId() const override  { RETURN_QUICK_DECLARE_CYCLE_STAT(FSplash, STATGROUP_Tickables); }
@@ -101,7 +101,7 @@ protected:
 protected:
 	FOculusHMD* OculusHMD;
 	FCustomPresent* CustomPresent;
-	TSharedPtr<FTicker> Ticker;
+	TSharedPtr<FTSTicker> Ticker;
 	int32 FramesOutstanding;
 	FCriticalSection RenderThreadLock;
 	FSettingsPtr Settings;
diff --git a/Engine/Plugins/Runtime/Steam/SteamSockets/Source/SteamSockets/Public/SteamSocketsSubsystem.h b/Engine/Plugins/Runtime/Steam/SteamSockets/Source/SteamSockets/Public/SteamSocketsSubsystem.h
--- a/Engine/Plugins/Runtime/Steam/SteamSockets/Source/SteamSockets/Public/SteamSocketsSubsystem.h
+++ b/Engine/Plugins/Runtime/Steam/SteamSockets/Source/SteamSockets/Public/SteamSocketsSubsystem.h
@@ -23,7 +23,7 @@ class USteamSocketsNetDriver;
  * This class can only be used with the SteamSocketsNetDriver and the SteamSocketsNetConnection classes.
  * This subsystem does not support mixing any other NetDriver/NetConnection format. Doing so will cause this protocol to not function.
  */
-class STEAMSOCKETS_API FSteamSocketsSubsystem : public ISocketSubsystem, public FTickerObjectBase, public FSelfRegisteringExec
+class STEAMSOCKETS_API FSteamSocketsSubsystem : public ISocketSubsystem, public FTSTickerObjectBase, public FSelfRegisteringExec
 {
 public:
 
diff --git a/Engine/Plugins/Tests/AutomationDriverTests/Source/AutomationDriverTests/Private/AutomationDriver.spec.cpp b/Engine/Plugins/Tests/AutomationDriverTests/Source/AutomationDriverTests/Private/AutomationDriver.spec.cpp
--- a/Engine/Plugins/Tests/AutomationDriverTests/Source/AutomationDriverTests/Private/AutomationDriver.spec.cpp
+++ b/Engine/Plugins/Tests/AutomationDriverTests/Source/AutomationDriverTests/Private/AutomationDriver.spec.cpp
@@ -385,7 +385,7 @@ void FAutomationDriverSpec::Define()
 						const TSharedPtr<int32, ESPMode::ThreadSafe> ExistenceChecker = WeakExistenceChecker.Pin();
 						if (ExistenceChecker.IsValid())
 						{
-							FTicker::GetCoreTicker().AddTicker(FTickerDelegate::CreateLambda([this, WeakExistenceChecker](float Delta){
+							FTSTicker::GetCoreTicker().AddTicker(FTickerDelegate::CreateLambda([this, WeakExistenceChecker](float Delta){
 								const TSharedPtr<int32, ESPMode::ThreadSafe> TickerExistenceChecker = WeakExistenceChecker.Pin();
 								if (TickerExistenceChecker.IsValid())
 								{
@@ -459,7 +459,7 @@ void FAutomationDriverSpec::Define()
 						const TSharedPtr<int32, ESPMode::ThreadSafe> ExistenceChecker = WeakExistenceChecker.Pin();
 						if (ExistenceChecker.IsValid())
 						{
-							FTicker::GetCoreTicker().AddTicker(FTickerDelegate::CreateLambda([this, WeakExistenceChecker](float Delta){
+							FTSTicker::GetCoreTicker().AddTicker(FTickerDelegate::CreateLambda([this, WeakExistenceChecker](float Delta){
 								const TSharedPtr<int32, ESPMode::ThreadSafe> TickerExistenceChecker = WeakExistenceChecker.Pin();
 								if (TickerExistenceChecker.IsValid())
 								{
@@ -1149,7 +1149,7 @@ void FAutomationDriverSpec::Define()
 							if (ExistenceChecker.IsValid())
 							{
 								SuiteWidget->RemoveContents();
-								FTicker::GetCoreTicker().AddTicker(FTickerDelegate::CreateLambda([this, WeakExistenceChecker](float Delta){
+								FTSTicker::GetCoreTicker().AddTicker(FTickerDelegate::CreateLambda([this, WeakExistenceChecker](float Delta){
 									const TSharedPtr<int32, ESPMode::ThreadSafe> TickerExistenceChecker = WeakExistenceChecker.Pin();
 									if (TickerExistenceChecker.IsValid())
 									{
@@ -1188,7 +1188,7 @@ void FAutomationDriverSpec::Define()
 							if (ExistenceChecker.IsValid())
 							{
 								SuiteWidget->RemoveContents();
-								FTicker::GetCoreTicker().AddTicker(FTickerDelegate::CreateLambda([this, WeakExistenceChecker](float Delta) {
+								FTSTicker::GetCoreTicker().AddTicker(FTickerDelegate::CreateLambda([this, WeakExistenceChecker](float Delta) {
 									const TSharedPtr<int32, ESPMode::ThreadSafe> TickerExistenceChecker = WeakExistenceChecker.Pin();
 									if (TickerExistenceChecker.IsValid())
 									{
@@ -1230,7 +1230,7 @@ void FAutomationDriverSpec::Define()
 							const TSharedPtr<int32, ESPMode::ThreadSafe> ExistenceChecker = WeakExistenceChecker.Pin();
 							if (ExistenceChecker.IsValid())
 							{
-								FTicker::GetCoreTicker().AddTicker(FTickerDelegate::CreateLambda([this, WeakExistenceChecker](float Delta){
+								FTSTicker::GetCoreTicker().AddTicker(FTickerDelegate::CreateLambda([this, WeakExistenceChecker](float Delta){
 									const TSharedPtr<int32, ESPMode::ThreadSafe> TickerExistenceChecker = WeakExistenceChecker.Pin();
 									if (TickerExistenceChecker.IsValid())
 									{
@@ -1269,7 +1269,7 @@ void FAutomationDriverSpec::Define()
 							const TSharedPtr<int32, ESPMode::ThreadSafe> ExistenceChecker = WeakExistenceChecker.Pin();
 							if (ExistenceChecker.IsValid())
 							{
-								FTicker::GetCoreTicker().AddTicker(FTickerDelegate::CreateLambda([this, WeakExistenceChecker](float Delta) {
+								FTSTicker::GetCoreTicker().AddTicker(FTickerDelegate::CreateLambda([this, WeakExistenceChecker](float Delta) {
 									const TSharedPtr<int32, ESPMode::ThreadSafe> TickerExistenceChecker = WeakExistenceChecker.Pin();
 									if (TickerExistenceChecker.IsValid())
 									{
@@ -1308,7 +1308,7 @@ void FAutomationDriverSpec::Define()
 							const TSharedPtr<int32, ESPMode::ThreadSafe> ExistenceChecker = WeakExistenceChecker.Pin();
 							if (ExistenceChecker.IsValid())
 							{
-								FTicker::GetCoreTicker().AddTicker(FTickerDelegate::CreateLambda([this, WeakExistenceChecker](float Delta){
+								FTSTicker::GetCoreTicker().AddTicker(FTickerDelegate::CreateLambda([this, WeakExistenceChecker](float Delta){
 									const TSharedPtr<int32, ESPMode::ThreadSafe> TickerExistenceChecker = WeakExistenceChecker.Pin();
 									if (TickerExistenceChecker.IsValid())
 									{
@@ -1346,7 +1346,7 @@ void FAutomationDriverSpec::Define()
 							const TSharedPtr<int32, ESPMode::ThreadSafe> ExistenceChecker = WeakExistenceChecker.Pin();
 							if (ExistenceChecker.IsValid())
 							{
-								FTicker::GetCoreTicker().AddTicker(FTickerDelegate::CreateLambda([this, WeakExistenceChecker](float Delta) {
+								FTSTicker::GetCoreTicker().AddTicker(FTickerDelegate::CreateLambda([this, WeakExistenceChecker](float Delta) {
 									const TSharedPtr<int32, ESPMode::ThreadSafe> TickerExistenceChecker = WeakExistenceChecker.Pin();
 									if (TickerExistenceChecker.IsValid())
 									{
diff --git a/Engine/Plugins/VirtualProduction/CameraCalibration/Source/CameraCalibrationEditor/Private/AssetEditor/CameraCalibrationStepsController.cpp b/Engine/Plugins/VirtualProduction/CameraCalibration/Source/CameraCalibrationEditor/Private/AssetEditor/CameraCalibrationStepsController.cpp
--- a/Engine/Plugins/VirtualProduction/CameraCalibration/Source/CameraCalibrationEditor/Private/AssetEditor/CameraCalibrationStepsController.cpp
+++ b/Engine/Plugins/VirtualProduction/CameraCalibration/Source/CameraCalibrationEditor/Private/AssetEditor/CameraCalibrationStepsController.cpp
@@ -113,7 +113,7 @@ FCameraCalibrationStepsController::FCameraCalibrationStepsController(TWeakPtr<FC
 	check(CameraCalibrationToolkit.IsValid());
 	check(LensFile.IsValid());
 
-	TickerHandle = FTicker::GetCoreTicker().AddTicker(
+	TickerHandle = FTSTicker::GetCoreTicker().AddTicker(
 		FTickerDelegate::CreateRaw(this, &FCameraCalibrationStepsController::OnTick), 0.0f);
 }
 
@@ -121,7 +121,7 @@ FCameraCalibrationStepsController::~FCameraCalibrationStepsController()
 {
 	if (TickerHandle.IsValid())
 	{
-		FTicker::GetCoreTicker().RemoveTicker(TickerHandle);
+		FTSTicker::GetCoreTicker().RemoveTicker(TickerHandle);
 	}
 
 	Cleanup();
diff --git a/Engine/Plugins/VirtualProduction/RemoteControl/Source/WebRemoteControl/Private/RemoteControlWebSocketServer.cpp b/Engine/Plugins/VirtualProduction/RemoteControl/Source/WebRemoteControl/Private/RemoteControlWebSocketServer.cpp
--- a/Engine/Plugins/VirtualProduction/RemoteControl/Source/WebRemoteControl/Private/RemoteControlWebSocketServer.cpp
+++ b/Engine/Plugins/VirtualProduction/RemoteControl/Source/WebRemoteControl/Private/RemoteControlWebSocketServer.cpp
@@ -93,7 +93,7 @@ bool FRCWebSocketServer::Start(uint32 Port, TSharedPtr<FWebsocketMessageRouter>
 	}
 
 	Router = MoveTemp(InRouter);
-	TickerHandle = FTicker::GetCoreTicker().AddTicker(FTickerDelegate::CreateRaw(this, &FRCWebSocketServer::Tick));
+	TickerHandle = FTSTicker::GetCoreTicker().AddTicker(FTickerDelegate::CreateRaw(this, &FRCWebSocketServer::Tick));
 
 	return true;
 }
@@ -122,7 +122,7 @@ void FRCWebSocketServer::Stop()
 {
 	if (TickerHandle.IsValid())
 	{
-		FTicker::GetCoreTicker().RemoveTicker(TickerHandle);
+		FTSTicker::GetCoreTicker().RemoveTicker(TickerHandle);
 	}
 	Router.Reset();
 	Server.Reset();
diff --git a/Engine/Plugins/VirtualProduction/StageMonitoring/Source/StageDataProvider/Private/StageDataProvider.cpp b/Engine/Plugins/VirtualProduction/StageMonitoring/Source/StageDataProvider/Private/StageDataProvider.cpp
--- a/Engine/Plugins/VirtualProduction/StageMonitoring/Source/StageDataProvider/Private/StageDataProvider.cpp
+++ b/Engine/Plugins/VirtualProduction/StageMonitoring/Source/StageDataProvider/Private/StageDataProvider.cpp
@@ -85,7 +85,7 @@ void FStageDataProvider::Start()
 		if (ProviderEndpoint.IsValid())
 		{
 			ProviderEndpoint->Subscribe<FStageProviderDiscoveryMessage>();
-			TickerHandle = FTicker::GetCoreTicker().AddTicker(FTickerDelegate::CreateRaw(this, &FStageDataProvider::Tick), 0.0f);
+			TickerHandle = FTSTicker::GetCoreTicker().AddTicker(FTickerDelegate::CreateRaw(this, &FStageDataProvider::Tick), 0.0f);
 
 			bIsActive = true;
 		}
@@ -101,7 +101,7 @@ void FStageDataProvider::Stop()
 
 		if (TickerHandle.IsValid())
 		{
-			FTicker::GetCoreTicker().RemoveTicker(TickerHandle);
+			FTSTicker::GetCoreTicker().RemoveTicker(TickerHandle);
 		}
 
 		if (ProviderEndpoint.IsValid())
diff --git a/Engine/Plugins/VirtualProduction/StageMonitoring/Source/StageMonitor/Private/StageMonitor.cpp b/Engine/Plugins/VirtualProduction/StageMonitoring/Source/StageMonitor/Private/StageMonitor.cpp
--- a/Engine/Plugins/VirtualProduction/StageMonitoring/Source/StageMonitor/Private/StageMonitor.cpp
+++ b/Engine/Plugins/VirtualProduction/StageMonitoring/Source/StageMonitor/Private/StageMonitor.cpp
@@ -57,7 +57,7 @@ void FStageMonitor::Start()
 			MonitorEndpoint->Subscribe<FStageProviderDiscoveryResponseMessage>();
 		}
 
-		TickerHandle = FTicker::GetCoreTicker().AddTicker(FTickerDelegate::CreateRaw(this, &FStageMonitor::Tick), 0.0f);
+		TickerHandle = FTSTicker::GetCoreTicker().AddTicker(FTickerDelegate::CreateRaw(this, &FStageMonitor::Tick), 0.0f);
 
 		bIsActive = true;
 	}
@@ -75,7 +75,7 @@ void FStageMonitor::Stop()
 
 	if (TickerHandle.IsValid())
 	{
-		FTicker::GetCoreTicker().RemoveTicker(TickerHandle);
+		FTSTicker::GetCoreTicker().RemoveTicker(TickerHandle);
 		TickerHandle.Reset();
 	}
 
diff --git a/Engine/Plugins/VirtualProduction/TimedDataMonitor/Source/TimedDataMonitor/Private/TimedDataMonitorCalibration.cpp b/Engine/Plugins/VirtualProduction/TimedDataMonitor/Source/TimedDataMonitor/Private/TimedDataMonitorCalibration.cpp
--- a/Engine/Plugins/VirtualProduction/TimedDataMonitor/Source/TimedDataMonitor/Private/TimedDataMonitorCalibration.cpp
+++ b/Engine/Plugins/VirtualProduction/TimedDataMonitor/Source/TimedDataMonitor/Private/TimedDataMonitorCalibration.cpp
@@ -364,7 +364,7 @@ void FTimedDataMonitorCalibration::Reset()
 	NumberOfRetry = 0;
 	if (TickerHandle.IsValid())
 	{
-		FTicker::GetCoreTicker().RemoveTicker(TickerHandle);
+		FTSTicker::GetCoreTicker().RemoveTicker(TickerHandle);
 	}
 
 	if (OnCalibrationCompleted.IsBound())
@@ -407,7 +407,7 @@ bool FTimedDataMonitorCalibration::RunCalibrateWithTimecode_CheckForReset(float)
 			
 		FTickerDelegate Delegate = FTickerDelegate::CreateRaw(this, &FTimedDataMonitorCalibration::RunCalibrateWithTimecode);
 		float Timer = CalibrationParameters.AmountOfSecondsToWaitAfterStatisticReset;
-		TickerHandle = FTicker::GetCoreTicker().AddTicker(Delegate, Timer);
+		TickerHandle = FTSTicker::GetCoreTicker().AddTicker(Delegate, Timer);
 	}
 	else
 	{
@@ -432,7 +432,7 @@ bool FTimedDataMonitorCalibration::RunCalibrateWithTimecode(float)
 			bCallcallback = false;
 			FTickerDelegate Delegate = FTickerDelegate::CreateRaw(this, &FTimedDataMonitorCalibration::RunCalibrateWithTimecode_CheckForReset);
 			float Timer = 1.f; // todo find a better way to handle this. how much time should we wait...
-			TickerHandle = FTicker::GetCoreTicker().AddTicker(Delegate, Timer);
+			TickerHandle = FTSTicker::GetCoreTicker().AddTicker(Delegate, Timer);
 		}
 
 	}
diff --git a/Engine/Source/Developer/Android/AndroidTargetPlatform/Private/AndroidTargetPlatform.cpp b/Engine/Source/Developer/Android/AndroidTargetPlatform/Private/AndroidTargetPlatform.cpp
--- a/Engine/Source/Developer/Android/AndroidTargetPlatform/Private/AndroidTargetPlatform.cpp
+++ b/Engine/Source/Developer/Android/AndroidTargetPlatform/Private/AndroidTargetPlatform.cpp
@@ -42,7 +42,7 @@ class FModuleManager;
 class FScopeLock;
 class FStaticMeshLODSettings;
 class FTargetDeviceId;
-class FTicker;
+class FTSTicker;
 class IAndroidDeviceDetectionModule;
 class UTexture;
 class UTextureLODSettings;
@@ -203,13 +203,13 @@ FAndroidTargetPlatform::FAndroidTargetPlatform(bool bInIsClient, const TCHAR* Fl
 #endif
 
 	TickDelegate = FTickerDelegate::CreateRaw(this, &FAndroidTargetPlatform::HandleTicker);
-	TickDelegateHandle = FTicker::GetCoreTicker().AddTicker(TickDelegate, 4.0f);
+	TickDelegateHandle = FTSTicker::GetCoreTicker().AddTicker(TickDelegate, 4.0f);
 }
 
 
 FAndroidTargetPlatform::~FAndroidTargetPlatform()
 {
-	 FTicker::GetCoreTicker().RemoveTicker(TickDelegateHandle);
+	 FTSTicker::GetCoreTicker().RemoveTicker(TickDelegateHandle);
 }
 
 FAndroidTargetDevicePtr FAndroidTargetPlatform::CreateTargetDevice(const ITargetPlatform& InTargetPlatform, const FString& InSerialNumber, const FString& InAndroidVariant) const
diff --git a/Engine/Source/Developer/AutomationController/Private/AutomationCommandline.cpp b/Engine/Source/Developer/AutomationController/Private/AutomationCommandline.cpp
--- a/Engine/Source/Developer/AutomationController/Private/AutomationCommandline.cpp
+++ b/Engine/Source/Developer/AutomationController/Private/AutomationCommandline.cpp
@@ -15,6 +15,7 @@
 #include "Misc/FileHelper.h"
 #include "AssetRegistryModule.h"
 #include "AutomationControllerSettings.h"
+#include "Containers/Ticker.h"
 
 DEFINE_LOG_CATEGORY_STATIC(LogAutomationCommandLine, Log, All);
 
@@ -75,7 +76,7 @@ public:
 			TestsRefreshedHandle = AutomationController->OnTestsRefreshed().AddRaw(this, &FAutomationExecCmd::HandleRefreshTestCallback);
 		}
 
-		TickHandler = FTicker::GetCoreTicker().AddTicker(FTickerDelegate::CreateRaw(this, &FAutomationExecCmd::Tick));
+		TickHandler = FTSTicker::GetCoreTicker().AddTicker(FTickerDelegate::CreateRaw(this, &FAutomationExecCmd::Tick));
 
 		int32 NumTestLoops = 1;
 		FParse::Value(FCommandLine::Get(), TEXT("TestLoops="), NumTestLoops);
@@ -103,7 +104,7 @@ public:
 			AutomationController->OnTestsRefreshed().RemoveAll(this);
 		}
 
-		FTicker::GetCoreTicker().RemoveTicker(TickHandler);
+		FTSTicker::GetCoreTicker().RemoveTicker(TickHandler);
 	}
 
 	bool IsTestingComplete()
@@ -688,7 +689,7 @@ private:
 	FGuid SessionID;
 
 	//so we can release control of the app and just get ticked like all other systems
-	FDelegateHandle TickHandler;
+	FTSTicker::FDelegateHandle TickHandler;
 
 	//Extra commandline params
 	FString StringCommand;
diff --git a/Engine/Source/Developer/AutomationDriver/Private/StepExecutor.cpp b/Engine/Source/Developer/AutomationDriver/Private/StepExecutor.cpp
--- a/Engine/Source/Developer/AutomationDriver/Private/StepExecutor.cpp
+++ b/Engine/Source/Developer/AutomationDriver/Private/StepExecutor.cpp
@@ -69,7 +69,7 @@ public:
 				if (Executor.IsValid())
 				{
 					const int32 StepIndex = 0;
-					FTicker::GetCoreTicker().AddTicker(FTickerDelegate::CreateThreadSafeSP(Executor.ToSharedRef(), &FStepExecutor::ExecuteStep, StepIndex), 0);
+					FTSTicker::GetCoreTicker().AddTicker(FTickerDelegate::CreateThreadSafeSP(Executor.ToSharedRef(), &FStepExecutor::ExecuteStep, StepIndex), 0);
 				}
 			}
 		);
@@ -145,7 +145,7 @@ private:
 
 		StepTotalProcessTime += FTimespan::FromSeconds(Delay);
 		LastDelay = Delay;
-		FTicker::GetCoreTicker().AddTicker(FTickerDelegate::CreateThreadSafeSP(this, &FStepExecutor::ExecuteStep, StepIndex), Delay);
+		FTSTicker::GetCoreTicker().AddTicker(FTickerDelegate::CreateThreadSafeSP(this, &FStepExecutor::ExecuteStep, StepIndex), Delay);
 
 		return false;
 	}
diff --git a/Engine/Source/Developer/CollectionManager/Private/CollectionManager.cpp b/Engine/Source/Developer/CollectionManager/Private/CollectionManager.cpp
--- a/Engine/Source/Developer/CollectionManager/Private/CollectionManager.cpp
+++ b/Engine/Source/Developer/CollectionManager/Private/CollectionManager.cpp
@@ -293,12 +293,12 @@ FCollectionManager::FCollectionManager()
 		CollectionFileCaches[CacheIdx] = MakeShareable(new DirectoryWatcher::FFileCache(FileCacheConfig));
 	}
 
-	TickFileCacheDelegateHandle = FTicker::GetCoreTicker().AddTicker(FTickerDelegate::CreateRaw(this, &FCollectionManager::TickFileCache), 1.0f);
+	TickFileCacheDelegateHandle = FTSTicker::GetCoreTicker().AddTicker(FTickerDelegate::CreateRaw(this, &FCollectionManager::TickFileCache), 1.0f);
 }
 
 FCollectionManager::~FCollectionManager()
 {
-	FTicker::GetCoreTicker().RemoveTicker(TickFileCacheDelegateHandle);
+	FTSTicker::GetCoreTicker().RemoveTicker(TickFileCacheDelegateHandle);
 }
 
 bool FCollectionManager::HasCollections() const
diff --git a/Engine/Source/Developer/CollectionManager/Private/CollectionManager.h b/Engine/Source/Developer/CollectionManager/Private/CollectionManager.h
--- a/Engine/Source/Developer/CollectionManager/Private/CollectionManager.h
+++ b/Engine/Source/Developer/CollectionManager/Private/CollectionManager.h
@@ -231,7 +231,7 @@ private:
 	TSharedPtr<DirectoryWatcher::FFileCache> CollectionFileCaches[ECollectionShareType::CST_All];
 
 	/** Delegate handle for the TickFileCache function */
-	FDelegateHandle TickFileCacheDelegateHandle;
+	FTSTicker::FDelegateHandle TickFileCacheDelegateHandle;
 
 	/** A map of collection names to FCollection objects */
 	TMap<FCollectionNameType, TSharedRef<FCollection>> AvailableCollections;
diff --git a/Engine/Source/Developer/Datasmith/DatasmithExporter/Private/DatasmithExporterManager.cpp b/Engine/Source/Developer/Datasmith/DatasmithExporter/Private/DatasmithExporterManager.cpp
--- a/Engine/Source/Developer/Datasmith/DatasmithExporter/Private/DatasmithExporterManager.cpp
+++ b/Engine/Source/Developer/Datasmith/DatasmithExporter/Private/DatasmithExporterManager.cpp
@@ -171,7 +171,7 @@ uint32 FDatasmithGameThread::Run()
 	{
 		FTaskGraphInterface::Get().ProcessThreadUntilIdle(ENamedThreads::GameThread);
 		FStats::AdvanceFrame(false);
-		FTicker::GetCoreTicker().Tick(FApp::GetDeltaTime());
+		FTSTicker::GetCoreTicker().Tick(FApp::GetDeltaTime());
 		FSlateApplication::Get().PumpMessages();
 		FSlateApplication::Get().Tick();
 
diff --git a/Engine/Source/Developer/DerivedDataCache/Private/HttpDerivedDataBackend.cpp b/Engine/Source/Developer/DerivedDataCache/Private/HttpDerivedDataBackend.cpp
--- a/Engine/Source/Developer/DerivedDataCache/Private/HttpDerivedDataBackend.cpp
+++ b/Engine/Source/Developer/DerivedDataCache/Private/HttpDerivedDataBackend.cpp
@@ -2079,7 +2079,7 @@ bool FHttpDerivedDataBackend::AcquireAccessToken()
 					//Schedule a refresh of the token ahead of expiry time (this will not work in commandlets)
 					if (!IsRunningCommandlet())
 					{
-						FTicker::GetCoreTicker().AddTicker(FTickerDelegate::CreateLambda(
+						FTSTicker::GetCoreTicker().AddTicker(FTickerDelegate::CreateLambda(
 							[this](float DeltaTime)
 							{
 								this->AcquireAccessToken();
diff --git a/Engine/Source/Developer/FunctionalTesting/Private/AutomationBlueprintFunctionLibrary.cpp b/Engine/Source/Developer/FunctionalTesting/Private/AutomationBlueprintFunctionLibrary.cpp
--- a/Engine/Source/Developer/FunctionalTesting/Private/AutomationBlueprintFunctionLibrary.cpp
+++ b/Engine/Source/Developer/FunctionalTesting/Private/AutomationBlueprintFunctionLibrary.cpp
@@ -404,7 +404,7 @@ public:
 	{
 		if (!bDeleteQueued)
 		{
-			FTicker::GetCoreTicker().AddTicker(TEXT("ScreenshotCleanup"), 0.1, [this](float) {
+			FTSTicker::GetCoreTicker().AddTicker(TEXT("ScreenshotCleanup"), 0.1, [this](float) {
 				delete this;
 				return false;
 				});
@@ -1180,7 +1180,7 @@ UAutomationEditorTask* UAutomationBlueprintFunctionLibrary::TakeHighResScreensho
 			Task->BindTask(MakeUnique<FScreenshotTakenState>());
 
 			// Delay taking the screenshot by a few frames			
-			FTicker::GetCoreTicker().AddTicker(TEXT("ScreenshotDelay"), Delay, [LevelViewport, ComparisonTolerance, ComparisonNotes, Filename, ResX, ResY, bMaskEnabled, bCaptureHDR](float) {
+			FTSTicker::GetCoreTicker().AddTicker(TEXT("ScreenshotDelay"), Delay, [LevelViewport, ComparisonTolerance, ComparisonNotes, Filename, ResX, ResY, bMaskEnabled, bCaptureHDR](float) {
 					FHighResScreenshotConfig& HighResScreenshotConfig = GetHighResScreenshotConfig();
 					HighResScreenshotConfig.SetResolution(ResX, ResY);
 					HighResScreenshotConfig.SetFilename(Filename);
diff --git a/Engine/Source/Developer/HierarchicalLODUtilities/Public/HierarchicalLODProxyProcessor.h b/Engine/Source/Developer/HierarchicalLODUtilities/Public/HierarchicalLODProxyProcessor.h
--- a/Engine/Source/Developer/HierarchicalLODUtilities/Public/HierarchicalLODProxyProcessor.h
+++ b/Engine/Source/Developer/HierarchicalLODUtilities/Public/HierarchicalLODProxyProcessor.h
@@ -13,15 +13,15 @@
 class ALODActor;
 class UHLODProxy;
 
-class HIERARCHICALLODUTILITIES_API FHierarchicalLODProxyProcessor : public FTickerObjectBase
+class HIERARCHICALLODUTILITIES_API FHierarchicalLODProxyProcessor : public FTSTickerObjectBase
 {
 public:		
 	FHierarchicalLODProxyProcessor();
 	~FHierarchicalLODProxyProcessor();
 	
-	/** Begin FTickerObjectBase */
+	/** Begin FTSTickerObjectBase */
 	virtual bool Tick(float DeltaTime) override;
-	/** End FTickerObjectBase */
+	/** End FTSTickerObjectBase */
 
 	/**
 	* AddProxyJob
diff --git a/Engine/Source/Developer/HotReload/Private/HotReload.cpp b/Engine/Source/Developer/HotReload/Private/HotReload.cpp
--- a/Engine/Source/Developer/HotReload/Private/HotReload.cpp
+++ b/Engine/Source/Developer/HotReload/Private/HotReload.cpp
@@ -51,6 +51,8 @@
 #include "Editor.h"
 #endif
 
+#include "Containers/Ticker.h"
+
 DEFINE_LOG_CATEGORY(LogHotReload);
 
 #define LOCTEXT_NAMESPACE "HotReload"
@@ -233,7 +235,7 @@ private:
 #endif
 
 	/**
-	 * Tick function for FTicker: checks for re-loaded modules and does hot-reload from IDE
+	 * Tick function for FTSTicker: checks for re-loaded modules and does hot-reload from IDE
 	 */
 	bool Tick(float DeltaTime);
 
@@ -310,11 +312,11 @@ private:
 	/** Callback registered with PluginManager to know if any new plugins have been created */
 	void PluginMountedCallback(IPlugin& Plugin);
 
-	/** FTicker delegate (hot-reload from IDE) */
+	/** FTSTicker delegate (hot-reload from IDE) */
 	FTickerDelegate TickerDelegate;
 
 	/** Handle to the registered TickerDelegate */
-	FDelegateHandle TickerDelegateHandle;
+	FTSTicker::FDelegateHandle TickerDelegateHandle;
 
 	/** Handle to the registered delegate above */
 	TMap<FString, FDelegateHandle> BinariesFolderChangedDelegateHandles;
@@ -549,7 +551,7 @@ void FHotReloadModule::StartupModule()
 
 	// Register hot-reload from IDE ticker
 	TickerDelegate = FTickerDelegate::CreateRaw(this, &FHotReloadModule::Tick);
-	TickerDelegateHandle = FTicker::GetCoreTicker().AddTicker(TickerDelegate);
+	TickerDelegateHandle = FTSTicker::GetCoreTicker().AddTicker(TickerDelegate);
 
 	FModuleManager::Get().OnModulesChanged().AddRaw(this, &FHotReloadModule::ModulesChangedCallback);
 
@@ -558,7 +560,7 @@ void FHotReloadModule::StartupModule()
 
 void FHotReloadModule::ShutdownModule()
 {
-	FTicker::GetCoreTicker().RemoveTicker(TickerDelegateHandle);
+	FTSTicker::GetCoreTicker().RemoveTicker(TickerDelegateHandle);
 	ShutdownHotReloadWatcher();
 
 	UEHotReload_Private::DeleteFileThatIndicatesEditorRunIfNeeded();
diff --git a/Engine/Source/Developer/IOS/IOSPlatformEditor/Private/IOSTargetSettingsCustomization.cpp b/Engine/Source/Developer/IOS/IOSPlatformEditor/Private/IOSTargetSettingsCustomization.cpp
--- a/Engine/Source/Developer/IOS/IOSPlatformEditor/Private/IOSTargetSettingsCustomization.cpp
+++ b/Engine/Source/Developer/IOS/IOSPlatformEditor/Private/IOSTargetSettingsCustomization.cpp
@@ -102,7 +102,7 @@ FIOSTargetSettingsCustomization::~FIOSTargetSettingsCustomization()
 	if (IPPProcess.IsValid())
 	{
 		IPPProcess = NULL;
-		FTicker::GetCoreTicker().RemoveTicker(TickerHandle);
+		FTSTicker::GetCoreTicker().RemoveTicker(TickerHandle);
 	}
 
 	FIOSPlatformEditorModule::OnSelect.RemoveAll(this);
@@ -1189,7 +1189,7 @@ void FIOSTargetSettingsCustomization::FindRequiredFiles()
 	OutputMessage = TEXT("");
 	IPPProcess->OnOutput().BindStatic(&OnOutput);
 	IPPProcess->Launch();
-	TickerHandle = FTicker::GetCoreTicker().AddTicker(FTickerDelegate::CreateRaw(this, &FIOSTargetSettingsCustomization::UpdateStatusDelegate), 1.0f);
+	TickerHandle = FTSTicker::GetCoreTicker().AddTicker(FTickerDelegate::CreateRaw(this, &FIOSTargetSettingsCustomization::UpdateStatusDelegate), 1.0f);
 	if (ProvisionInfoSwitcher.IsValid())
 	{
 		ProvisionInfoSwitcher->SetActiveWidgetIndex(0);
@@ -1269,7 +1269,7 @@ FReply FIOSTargetSettingsCustomization::OnInstallProvisionClicked()
 		OutputMessage = TEXT("");
 		IPPProcess->OnOutput().BindStatic(&OnOutput);
 		IPPProcess->Launch();
-		TickerHandle = FTicker::GetCoreTicker().AddTicker(FTickerDelegate::CreateRaw(this, &FIOSTargetSettingsCustomization::UpdateStatusDelegate), 10.0f);
+		TickerHandle = FTSTicker::GetCoreTicker().AddTicker(FTickerDelegate::CreateRaw(this, &FIOSTargetSettingsCustomization::UpdateStatusDelegate), 10.0f);
 		if (ProvisionInfoSwitcher.IsValid())
 		{
 			ProvisionInfoSwitcher->SetActiveWidgetIndex(1);
@@ -1326,7 +1326,7 @@ FReply FIOSTargetSettingsCustomization::OnInstallCertificateClicked()
 		OutputMessage = TEXT("");
 		IPPProcess->OnOutput().BindStatic(&OnOutput);
 		IPPProcess->Launch();
-		TickerHandle = FTicker::GetCoreTicker().AddTicker(FTickerDelegate::CreateRaw(this, &FIOSTargetSettingsCustomization::UpdateStatusDelegate), 10.0f);
+		TickerHandle = FTSTicker::GetCoreTicker().AddTicker(FTickerDelegate::CreateRaw(this, &FIOSTargetSettingsCustomization::UpdateStatusDelegate), 10.0f);
 		if (CertificateInfoSwitcher.IsValid())
 		{
 			CertificateInfoSwitcher->SetActiveWidgetIndex(1);
@@ -1402,7 +1402,7 @@ FReply FIOSTargetSettingsCustomization::OnGenerateSSHKey()
 	OutputMessage = TEXT("");
 	IPPProcess = MakeShareable(new FMonitoredProcess(CmdExe, CommandLine, false, false));
 	IPPProcess->Launch();
-	TickerHandle = FTicker::GetCoreTicker().AddTicker(FTickerDelegate::CreateRaw(this, &FIOSTargetSettingsCustomization::UpdateStatusDelegate), 10.0f);
+	TickerHandle = FTSTicker::GetCoreTicker().AddTicker(FTickerDelegate::CreateRaw(this, &FIOSTargetSettingsCustomization::UpdateStatusDelegate), 10.0f);
 	RunningIPPProcess = true;
 
 	return FReply::Handled();
diff --git a/Engine/Source/Developer/IOS/IOSTargetPlatform/Private/IOSDeviceHelper.cpp b/Engine/Source/Developer/IOS/IOSTargetPlatform/Private/IOSDeviceHelper.cpp
--- a/Engine/Source/Developer/IOS/IOSTargetPlatform/Private/IOSDeviceHelper.cpp
+++ b/Engine/Source/Developer/IOS/IOSTargetPlatform/Private/IOSDeviceHelper.cpp
@@ -349,7 +349,7 @@ void FIOSDeviceHelper::Initialize(bool bIsTVOS)
     {
         // add the message pump
         TickDelegate = FTickerDelegate::CreateStatic(MessageTickDelegate);
-        FTicker::GetCoreTicker().AddTicker(TickDelegate, 5.0f);
+        FTSTicker::GetCoreTicker().AddTicker(TickDelegate, 5.0f);
 
         // kick off a thread to query for connected devices
         QueryTask = new FDeviceQueryTask();
diff --git a/Engine/Source/Developer/MeshMergeUtilities/Private/ProxyGenerationProcessor.h b/Engine/Source/Developer/MeshMergeUtilities/Private/ProxyGenerationProcessor.h
--- a/Engine/Source/Developer/MeshMergeUtilities/Private/ProxyGenerationProcessor.h
+++ b/Engine/Source/Developer/MeshMergeUtilities/Private/ProxyGenerationProcessor.h
@@ -12,7 +12,7 @@ class FMeshMergeUtilities;
 
 DECLARE_LOG_CATEGORY_EXTERN(LogMeshMerging, Verbose, All);
 
-class FProxyGenerationProcessor : FTickerObjectBase
+class FProxyGenerationProcessor : FTSTickerObjectBase
 {
 public:
 	FProxyGenerationProcessor(const FMeshMergeUtilities* InOwner);
diff --git a/Engine/Source/Developer/Profiler/Private/ProfilerManager.cpp b/Engine/Source/Developer/Profiler/Private/ProfilerManager.cpp
--- a/Engine/Source/Developer/Profiler/Private/ProfilerManager.cpp
+++ b/Engine/Source/Developer/Profiler/Private/ProfilerManager.cpp
@@ -97,7 +97,7 @@ void FProfilerManager::PostConstructor()
 {
 	// Register tick functions.
 	OnTick = FTickerDelegate::CreateSP( this, &FProfilerManager::Tick );
-	OnTickHandle = FTicker::GetCoreTicker().AddTicker( OnTick, 1.0f );
+	OnTickHandle = FTSTicker::GetCoreTicker().AddTicker( OnTick, 1.0f );
 
 	// Create profiler client.
 	ProfilerClient = FModuleManager::GetModuleChecked<IProfilerClientModule>("ProfilerClient").CreateProfilerClient();
@@ -139,7 +139,7 @@ FProfilerManager::~FProfilerManager()
 	FProfilerCommands::Unregister();
 
 	// Unregister tick function.
-	FTicker::GetCoreTicker().RemoveTicker( OnTickHandle );
+	FTSTicker::GetCoreTicker().RemoveTicker( OnTickHandle );
 
 	// Remove ourselves from the session manager.
 	if (SessionManager.IsValid())
diff --git a/Engine/Source/Developer/Profiler/Private/ProfilerManager.h b/Engine/Source/Developer/Profiler/Private/ProfilerManager.h
--- a/Engine/Source/Developer/Profiler/Private/ProfilerManager.h
+++ b/Engine/Source/Developer/Profiler/Private/ProfilerManager.h
@@ -11,6 +11,7 @@
 #include "ProfilerDataSource.h"
 #include "ProfilerSession.h"
 #include "ProfilerCommands.h"
+#include "Containers/Ticker.h"
 
 class SProfilerWindow;
 
@@ -457,7 +458,7 @@ protected:
 	FTickerDelegate OnTick;
 
 	/** Handle to the registered OnTick. */
-	FDelegateHandle OnTickHandle;
+	FTSTicker::FDelegateHandle OnTickHandle;
 
 	/** A weak pointer to the profiler window. */
 	TWeakPtr<class SProfilerWindow> ProfilerWindow;
diff --git a/Engine/Source/Developer/Profiler/Private/ProfilerRawStatsForThreadView.cpp b/Engine/Source/Developer/Profiler/Private/ProfilerRawStatsForThreadView.cpp
--- a/Engine/Source/Developer/Profiler/Private/ProfilerRawStatsForThreadView.cpp
+++ b/Engine/Source/Developer/Profiler/Private/ProfilerRawStatsForThreadView.cpp
@@ -21,7 +21,7 @@ FRawProfilerSession::FRawProfilerSession( const FString& InRawStatsFileFileath )
 
 FRawProfilerSession::~FRawProfilerSession()
 {
-	FTicker::GetCoreTicker().RemoveTicker( OnTickHandle );
+	FTSTicker::GetCoreTicker().RemoveTicker( OnTickHandle );
 }
 
 bool FRawProfilerSession::HandleTicker( float DeltaTime )
@@ -351,7 +351,7 @@ void FRawProfilerSession::PrepareLoading()
 	const int64 AllocatedSize = ProfilerStream.GetAllocatedSize();
 
 	// We have the whole metadata and basic information about the raw stats file, start ticking the profiler session.
-	//OnTickHandle = FTicker::GetCoreTicker().AddTicker( OnTick, 0.25f );
+	//OnTickHandle = FTSTicker::GetCoreTicker().AddTicker( OnTick, 0.25f );
 
 #if	0
 	if( SessionType == EProfilerSessionTypes::OfflineRaw )
diff --git a/Engine/Source/Developer/Profiler/Private/ProfilerSession.cpp b/Engine/Source/Developer/Profiler/Private/ProfilerSession.cpp
--- a/Engine/Source/Developer/Profiler/Private/ProfilerSession.cpp
+++ b/Engine/Source/Developer/Profiler/Private/ProfilerSession.cpp
@@ -68,7 +68,7 @@ FProfilerSession::FProfilerSession( const TSharedPtr<ISessionInstanceInfo>& InSe
 {
 	// Randomize creation time to test loading profiler captures with different creation time and different amount of data.
 	CreationTime = FDateTime::Now() += FTimespan( 0, 0, FMath::RandRange( 2, 8 ) );
-	OnTickHandle = FTicker::GetCoreTicker().AddTicker( OnTick );
+	OnTickHandle = FTSTicker::GetCoreTicker().AddTicker( OnTick );
 }
 
 FProfilerSession::FProfilerSession( const FString& InDataFilepath )
@@ -94,12 +94,12 @@ FProfilerSession::FProfilerSession( const FString& InDataFilepath )
 {
 	// Randomize creation time to test loading profiler captures with different creation time and different amount of data.
 	CreationTime = FDateTime::Now() += FTimespan( 0, 0, FMath::RandRange( 2, 8 ) );
-	OnTickHandle = FTicker::GetCoreTicker().AddTicker( OnTick );
+	OnTickHandle = FTSTicker::GetCoreTicker().AddTicker( OnTick );
 }
 
 FProfilerSession::~FProfilerSession()
 {
-	FTicker::GetCoreTicker().RemoveTicker( OnTickHandle );	
+	FTSTicker::GetCoreTicker().RemoveTicker( OnTickHandle );	
 }
 
 const TSharedRef<FEventGraphData, ESPMode::ThreadSafe> FProfilerSession::GetEventGraphDataTotal() const
diff --git a/Engine/Source/Developer/Profiler/Private/ProfilerSession.h b/Engine/Source/Developer/Profiler/Private/ProfilerSession.h
--- a/Engine/Source/Developer/Profiler/Private/ProfilerSession.h
+++ b/Engine/Source/Developer/Profiler/Private/ProfilerSession.h
@@ -1057,7 +1057,7 @@ protected:
 	FTickerDelegate OnTick;
 
 	/** Handle to the registered OnTick. */
-	FDelegateHandle OnTickHandle;
+	FTSTicker::FDelegateHandle OnTickHandle;
 
 	/** The data provider which holds all the collected profiler samples. */
 	TSharedRef<IDataProvider> DataProvider;
diff --git a/Engine/Source/Developer/ProfilerClient/Private/ProfilerClientManager.cpp b/Engine/Source/Developer/ProfilerClient/Private/ProfilerClientManager.cpp
--- a/Engine/Source/Developer/ProfilerClient/Private/ProfilerClientManager.cpp
+++ b/Engine/Source/Developer/ProfilerClient/Private/ProfilerClientManager.cpp
@@ -49,7 +49,7 @@ FProfilerClientManager::FProfilerClientManager(const TSharedRef<IMessageBus, ESP
 	RetryTime = 5.f;
 
 	LoadConnection = nullptr;
-	MessageDelegateHandle = FTicker::GetCoreTicker().AddTicker(MessageDelegate, 0.1f);
+	MessageDelegateHandle = FTSTicker::GetCoreTicker().AddTicker(MessageDelegate, 0.1f);
 #endif
 }
 
@@ -115,7 +115,7 @@ void FProfilerClientManager::Track(const FGuid& Instance)
 		MessageEndpoint->Publish(FMessageEndpoint::MakeMessage<FProfilerServiceSubscribe>(ActiveSessionId, Instance), EMessageScope::Network);
 
 		RetryTime = 5.f;
-		TickDelegateHandle = FTicker::GetCoreTicker().AddTicker(TickDelegate, RetryTime);
+		TickDelegateHandle = FTSTicker::GetCoreTicker().AddTicker(TickDelegate, RetryTime);
 
 		UE_LOG(LogProfilerClient, Verbose, TEXT("Track Session: %s, Instance: %s"), *ActiveSessionId.ToString(), *Instance.ToString());
 	}
@@ -315,7 +315,7 @@ void FProfilerClientManager::LoadCapture(const FString& DataFilepath, const FGui
 	LoadConnection->LoadCapture(DataFilepath, this);
 
 	RetryTime = 0.05f;
-	TickDelegateHandle = FTicker::GetCoreTicker().AddTicker(TickDelegate, RetryTime);	
+	TickDelegateHandle = FTSTicker::GetCoreTicker().AddTicker(TickDelegate, RetryTime);	
 #endif
 }
 
@@ -837,8 +837,8 @@ void FProfilerClientManager::Shutdown()
 		UE_LOG(LogProfilerClient, Log, TEXT("File service-client transfer aborted: %s"), *It.Key());
 	}
 
-	FTicker::GetCoreTicker().RemoveTicker(MessageDelegateHandle);
-	FTicker::GetCoreTicker().RemoveTicker(TickDelegateHandle);
+	FTSTicker::GetCoreTicker().RemoveTicker(MessageDelegateHandle);
+	FTSTicker::GetCoreTicker().RemoveTicker(TickDelegateHandle);
 }
 
 
diff --git a/Engine/Source/Developer/ProfilerClient/Private/ProfilerClientManager.h b/Engine/Source/Developer/ProfilerClient/Private/ProfilerClientManager.h
--- a/Engine/Source/Developer/ProfilerClient/Private/ProfilerClientManager.h
+++ b/Engine/Source/Developer/ProfilerClient/Private/ProfilerClientManager.h
@@ -269,7 +269,7 @@ private:
 	FTickerDelegate TickDelegate;
 
 	/** Handle to the registered TickDelegate. */
-	FDelegateHandle TickDelegateHandle;
+	FTSTicker::FDelegateHandle TickDelegateHandle;
 
 	/** Amount of time between connection retries */
 	float RetryTime;
@@ -281,7 +281,7 @@ private:
 	FTickerDelegate MessageDelegate;
 
 	/** Handle to the registered MessageDelegate. */
-	FDelegateHandle MessageDelegateHandle;
+	FTSTicker::FDelegateHandle MessageDelegateHandle;
 
 	/** Handle to the registered OnShutdown for Message Bus. */
 	FDelegateHandle OnShutdownMessageBusDelegateHandle;
diff --git a/Engine/Source/Developer/ProfilerService/Private/ProfilerServiceManager.cpp b/Engine/Source/Developer/ProfilerService/Private/ProfilerServiceManager.cpp
--- a/Engine/Source/Developer/ProfilerService/Private/ProfilerServiceManager.cpp
+++ b/Engine/Source/Developer/ProfilerService/Private/ProfilerServiceManager.cpp
@@ -355,7 +355,7 @@ void FProfilerServiceManager::HandleServiceSubscribeMessage( const FProfilerServ
 		// Initiate the ping callback
 		if (ClientData.Num() == 1)
 		{
-			PingDelegateHandle = FTicker::GetCoreTicker().AddTicker(PingDelegate, 5.0f);
+			PingDelegateHandle = FTSTicker::GetCoreTicker().AddTicker(PingDelegate, 5.0f);
 		}
 	}
 #endif //STATS
@@ -383,7 +383,7 @@ void FProfilerServiceManager::HandleServiceUnsubscribeMessage( const FProfilerSe
 		// stop the ping messages if we have no clients
 		if (ClientData.Num() == 0)
 		{
-			FTicker::GetCoreTicker().RemoveTicker(PingDelegateHandle);
+			FTSTicker::GetCoreTicker().RemoveTicker(PingDelegateHandle);
 		}
 	}
 #endif //STATS
diff --git a/Engine/Source/Developer/ProfilerService/Private/ProfilerServiceManager.h b/Engine/Source/Developer/ProfilerService/Private/ProfilerServiceManager.h
--- a/Engine/Source/Developer/ProfilerService/Private/ProfilerServiceManager.h
+++ b/Engine/Source/Developer/ProfilerService/Private/ProfilerServiceManager.h
@@ -239,7 +239,7 @@ private:
 	FTickerDelegate PingDelegate;
 
 	/** Handle to the registered PingDelegate */
-	FDelegateHandle PingDelegateHandle;
+	FTSTicker::FDelegateHandle PingDelegateHandle;
 
 	/** Handle to the registered HandleNewFrame delegate */
 	FDelegateHandle NewFrameDelegateHandle;
diff --git a/Engine/Source/Developer/TargetDeviceServices/Private/Proxies/TargetDeviceProxyManager.cpp b/Engine/Source/Developer/TargetDeviceServices/Private/Proxies/TargetDeviceProxyManager.cpp
--- a/Engine/Source/Developer/TargetDeviceServices/Private/Proxies/TargetDeviceProxyManager.cpp
+++ b/Engine/Source/Developer/TargetDeviceServices/Private/Proxies/TargetDeviceProxyManager.cpp
@@ -25,7 +25,7 @@ FTargetDeviceProxyManager::FTargetDeviceProxyManager()
 	if (MessageEndpoint.IsValid())
 	{
 		TickDelegate = FTickerDelegate::CreateRaw(this, &FTargetDeviceProxyManager::HandleTicker);
-		TickDelegateHandle = FTicker::GetCoreTicker().AddTicker(TickDelegate, TARGET_DEVICE_SERVICES_PING_INTERVAL);
+		TickDelegateHandle = FTSTicker::GetCoreTicker().AddTicker(TickDelegate, TARGET_DEVICE_SERVICES_PING_INTERVAL);
 
 		SendPing();
 	}
@@ -34,7 +34,7 @@ FTargetDeviceProxyManager::FTargetDeviceProxyManager()
 
 FTargetDeviceProxyManager::~FTargetDeviceProxyManager()
 {
-	FTicker::GetCoreTicker().RemoveTicker(TickDelegateHandle);
+	FTSTicker::GetCoreTicker().RemoveTicker(TickDelegateHandle);
 	FMessageEndpoint::SafeRelease(MessageEndpoint);
 }
 
diff --git a/Engine/Source/Developer/TraceInsights/Private/Insights/ContextSwitches/ContextSwitchesProfilerManager.cpp b/Engine/Source/Developer/TraceInsights/Private/Insights/ContextSwitches/ContextSwitchesProfilerManager.cpp
--- a/Engine/Source/Developer/TraceInsights/Private/Insights/ContextSwitches/ContextSwitchesProfilerManager.cpp
+++ b/Engine/Source/Developer/TraceInsights/Private/Insights/ContextSwitches/ContextSwitchesProfilerManager.cpp
@@ -66,7 +66,7 @@ void FContextSwitchesProfilerManager::Initialize(IUnrealInsightsModule& Insights
 
 	// Register tick functions.
 	OnTick = FTickerDelegate::CreateSP(this, &FContextSwitchesProfilerManager::Tick);
-	OnTickHandle = FTicker::GetCoreTicker().AddTicker(OnTick, 0.0f);
+	OnTickHandle = FTSTicker::GetCoreTicker().AddTicker(OnTick, 0.0f);
 
 	FOnRegisterMajorTabExtensions* TimingProfilerLayoutExtension = InsightsModule.FindMajorTabLayoutExtension(FInsightsManagerTabs::TimingProfilerTabId);
 	if (TimingProfilerLayoutExtension)
@@ -91,7 +91,7 @@ void FContextSwitchesProfilerManager::Shutdown()
 	FInsightsManager::Get()->GetSessionChangedEvent().RemoveAll(this);
 
 	// Unregister tick function.
-	FTicker::GetCoreTicker().RemoveTicker(OnTickHandle);
+	FTSTicker::GetCoreTicker().RemoveTicker(OnTickHandle);
 
 	FContextSwitchesProfilerManager::Instance.Reset();
 }
diff --git a/Engine/Source/Developer/TraceInsights/Private/Insights/ContextSwitches/ContextSwitchesProfilerManager.h b/Engine/Source/Developer/TraceInsights/Private/Insights/ContextSwitches/ContextSwitchesProfilerManager.h
--- a/Engine/Source/Developer/TraceInsights/Private/Insights/ContextSwitches/ContextSwitchesProfilerManager.h
+++ b/Engine/Source/Developer/TraceInsights/Private/Insights/ContextSwitches/ContextSwitchesProfilerManager.h
@@ -68,7 +68,7 @@ private:
 	FTickerDelegate OnTick;
 
 	/** Handle to the registered OnTick. */
-	FDelegateHandle OnTickHandle;
+	FTSTicker::FDelegateHandle OnTickHandle;
 
 	/** A shared pointer to the global instance of the Context Switches Profiler manager. */
 	static TSharedPtr<FContextSwitchesProfilerManager> Instance;
diff --git a/Engine/Source/Developer/TraceInsights/Private/Insights/InsightsManager.cpp b/Engine/Source/Developer/TraceInsights/Private/Insights/InsightsManager.cpp
--- a/Engine/Source/Developer/TraceInsights/Private/Insights/InsightsManager.cpp
+++ b/Engine/Source/Developer/TraceInsights/Private/Insights/InsightsManager.cpp
@@ -158,7 +158,7 @@ void FInsightsManager::Initialize(IUnrealInsightsModule& InsightsModule)
 
 	// Register tick functions.
 	OnTick = FTickerDelegate::CreateSP(this, &FInsightsManager::Tick);
-	OnTickHandle = FTicker::GetCoreTicker().AddTicker(OnTick, 0.0f);
+	OnTickHandle = FTSTicker::GetCoreTicker().AddTicker(OnTick, 0.0f);
 
 	FInsightsCommands::Register();
 	BindCommands();
@@ -179,7 +179,7 @@ void FInsightsManager::Shutdown()
 	FInsightsCommands::Unregister();
 
 	// Unregister tick function.
-	FTicker::GetCoreTicker().RemoveTicker(OnTickHandle);
+	FTSTicker::GetCoreTicker().RemoveTicker(OnTickHandle);
 
 	// If the MessageLog module was already unloaded as part of the global Shutdown process, do not load it again.
 	if (FModuleManager::Get().IsModuleLoaded("MessageLog"))
diff --git a/Engine/Source/Developer/TraceInsights/Private/Insights/Tests/InsightsTestRunner.cpp b/Engine/Source/Developer/TraceInsights/Private/Insights/Tests/InsightsTestRunner.cpp
--- a/Engine/Source/Developer/TraceInsights/Private/Insights/Tests/InsightsTestRunner.cpp
+++ b/Engine/Source/Developer/TraceInsights/Private/Insights/Tests/InsightsTestRunner.cpp
@@ -74,14 +74,14 @@ void FInsightsTestRunner::Initialize(IUnrealInsightsModule& InsightsModule)
 
 	// Register tick functions.
 	OnTick = FTickerDelegate::CreateSP(this, &FInsightsTestRunner::Tick);
-	OnTickHandle = FTicker::GetCoreTicker().AddTicker(OnTick, 1.0f);
+	OnTickHandle = FTSTicker::GetCoreTicker().AddTicker(OnTick, 1.0f);
 }
 
 ////////////////////////////////////////////////////////////////////////////////////////////////////
 
 void FInsightsTestRunner::Shutdown()
 {
-	FTicker::GetCoreTicker().RemoveTicker(OnTickHandle);
+	FTSTicker::GetCoreTicker().RemoveTicker(OnTickHandle);
 }
 
 ////////////////////////////////////////////////////////////////////////////////////////////////////
diff --git a/Engine/Source/Developer/TraceInsights/Private/Insights/Tests/InsightsTestRunner.h b/Engine/Source/Developer/TraceInsights/Private/Insights/Tests/InsightsTestRunner.h
--- a/Engine/Source/Developer/TraceInsights/Private/Insights/Tests/InsightsTestRunner.h
+++ b/Engine/Source/Developer/TraceInsights/Private/Insights/Tests/InsightsTestRunner.h
@@ -4,6 +4,7 @@
 
 #include "CoreMinimal.h"
 #include "Logging/LogMacros.h"
+#include "Containers/Ticker.h"
 
 #if !UE_BUILD_SHIPPING && !WITH_EDITOR
 
@@ -28,7 +29,7 @@ public:
 	FTickerDelegate OnTick;
 
 	/** Handle to the registered OnTick. */
-	FDelegateHandle OnTickHandle;
+	FTSTicker::FDelegateHandle OnTickHandle;
 
 	bool Tick(float DeltaTime);
 
diff --git a/Engine/Source/Developer/Virtualization/Private/VirtualizationJupiterBackend.cpp b/Engine/Source/Developer/Virtualization/Private/VirtualizationJupiterBackend.cpp
--- a/Engine/Source/Developer/Virtualization/Private/VirtualizationJupiterBackend.cpp
+++ b/Engine/Source/Developer/Virtualization/Private/VirtualizationJupiterBackend.cpp
@@ -1759,7 +1759,7 @@ private:
 						//Schedule a refresh of the token ahead of expiry time (this will not work in commandlets)
 						if (!IsRunningCommandlet())
 						{
-							FTicker::GetCoreTicker().AddTicker(FTickerDelegate::CreateLambda(
+							FTSTicker::GetCoreTicker().AddTicker(FTickerDelegate::CreateLambda(
 								[this](float DeltaTime)
 								{
 									this->AcquireAccessToken();
diff --git a/Engine/Source/Editor/AnimGraph/Private/AnimGraphNode_LinkedInputPose.cpp b/Engine/Source/Editor/AnimGraph/Private/AnimGraphNode_LinkedInputPose.cpp
--- a/Engine/Source/Editor/AnimGraph/Private/AnimGraphNode_LinkedInputPose.cpp
+++ b/Engine/Source/Editor/AnimGraph/Private/AnimGraphNode_LinkedInputPose.cpp
@@ -342,7 +342,7 @@ void UAnimGraphNode_LinkedInputPose::PostPlacedNewNode()
 			return true;
 		});
 
-		FTicker::GetCoreTicker().AddTicker(FTickerDelegate::CreateLambda([WeakThis = TWeakObjectPtr<UAnimGraphNode_LinkedInputPose>(this)](float InDeltaTime)
+		FTSTicker::GetCoreTicker().AddTicker(FTickerDelegate::CreateLambda([WeakThis = TWeakObjectPtr<UAnimGraphNode_LinkedInputPose>(this)](float InDeltaTime)
 		{
 			QUICK_SCOPE_CYCLE_COUNTER(STAT_UAnimGraphNode_LinkedInputPose_PostPlacedNewNode);
 			if(UAnimGraphNode_LinkedInputPose* LinkedInputPoseNode = WeakThis.Get())
diff --git a/Engine/Source/Editor/Blutility/Private/EditorUtilityLibrary.cpp b/Engine/Source/Editor/Blutility/Private/EditorUtilityLibrary.cpp
--- a/Engine/Source/Editor/Blutility/Private/EditorUtilityLibrary.cpp
+++ b/Engine/Source/Editor/Blutility/Private/EditorUtilityLibrary.cpp
@@ -55,7 +55,7 @@ void UAsyncEditorDelay::Start(float InMinimumSeconds, int32 InMinimumFrames)
 {
 	EndFrame = GFrameCounter + InMinimumFrames;
 	EndTime = FApp::GetCurrentTime() + InMinimumSeconds;
-	FTicker::GetCoreTicker().AddTicker(FTickerDelegate::CreateUObject(this, &UAsyncEditorDelay::HandleComplete), 0);
+	FTSTicker::GetCoreTicker().AddTicker(FTickerDelegate::CreateUObject(this, &UAsyncEditorDelay::HandleComplete), 0);
 }
 
 bool UAsyncEditorDelay::HandleComplete(float DeltaTime)
@@ -97,7 +97,7 @@ void UAsyncEditorWaitForGameWorld::Start(int32 InIndex, bool InServer)
 {
 	Index = InIndex;
 	Server = InServer;
-	FTicker::GetCoreTicker().AddTicker(FTickerDelegate::CreateUObject(this, &UAsyncEditorWaitForGameWorld::OnTick), 0);
+	FTSTicker::GetCoreTicker().AddTicker(FTickerDelegate::CreateUObject(this, &UAsyncEditorWaitForGameWorld::OnTick), 0);
 }
 
 bool UAsyncEditorWaitForGameWorld::OnTick(float DeltaTime)
@@ -176,7 +176,7 @@ void UAsyncEditorOpenMapAndFocusActor::Start(FSoftObjectPath InMap, FString InFo
 	UWorld* World = GEditor ? GEditor->GetEditorWorldContext(false).World() : nullptr;
 	UKismetSystemLibrary::ExecuteConsoleCommand(World, FString::Printf(TEXT("Automate.OpenMapAndFocusActor %s %s"), *InMap.ToString(), *InFocusActorName));
 
-	FTicker::GetCoreTicker().AddTicker(FTickerDelegate::CreateUObject(this, &UAsyncEditorOpenMapAndFocusActor::OnTick), 0);
+	FTSTicker::GetCoreTicker().AddTicker(FTickerDelegate::CreateUObject(this, &UAsyncEditorOpenMapAndFocusActor::OnTick), 0);
 }
 
 bool UAsyncEditorOpenMapAndFocusActor::OnTick(float DeltaTime)
diff --git a/Engine/Source/Editor/Blutility/Private/EditorUtilitySubsystem.cpp b/Engine/Source/Editor/Blutility/Private/EditorUtilitySubsystem.cpp
--- a/Engine/Source/Editor/Blutility/Private/EditorUtilitySubsystem.cpp
+++ b/Engine/Source/Editor/Blutility/Private/EditorUtilitySubsystem.cpp
@@ -49,7 +49,7 @@ void UEditorUtilitySubsystem::Initialize(FSubsystemCollectionBase& Collection)
 		MainFrameModule.OnMainFrameCreationFinished().AddUObject(this, &UEditorUtilitySubsystem::MainFrameCreationFinished);
 	}
 
-	TickerHandle = FTicker::GetCoreTicker().AddTicker(FTickerDelegate::CreateUObject(this, &UEditorUtilitySubsystem::Tick), 0);
+	TickerHandle = FTSTicker::GetCoreTicker().AddTicker(FTickerDelegate::CreateUObject(this, &UEditorUtilitySubsystem::Tick), 0);
 }
 
 void UEditorUtilitySubsystem::Deinitialize()
@@ -59,7 +59,7 @@ void UEditorUtilitySubsystem::Deinitialize()
 		IMainFrameModule::Get().OnMainFrameCreationFinished().RemoveAll(this);
 	}
 
-	FTicker::GetCoreTicker().RemoveTicker(TickerHandle);
+	FTSTicker::GetCoreTicker().RemoveTicker(TickerHandle);
 
 	IConsoleManager::Get().UnregisterConsoleObject(RunTaskCommandObject);
 }
diff --git a/Engine/Source/Editor/ContentBrowser/Private/FrontendFilters.cpp b/Engine/Source/Editor/ContentBrowser/Private/FrontendFilters.cpp
--- a/Engine/Source/Editor/ContentBrowser/Private/FrontendFilters.cpp
+++ b/Engine/Source/Editor/ContentBrowser/Private/FrontendFilters.cpp
@@ -1180,7 +1180,7 @@ void FFrontendFilter_InUseByLoadedLevels::RegisterDelayedRefresh(float DelayInSe
 	// The Editor might be unresponsive during heavy asset compilation so we 
 	// not only need a delay, but also a minimum amount of frames
 	// to pass until we call the actual refresh.
-	DelayedRefreshHandle = FTicker::GetCoreTicker().AddTicker(
+	DelayedRefreshHandle = FTSTicker::GetCoreTicker().AddTicker(
 		TEXT("FFrontendFilter_InUseByLoadedLevels"),
 		0.0f,
 		[this, FireInTickCount = 16, DelayInSeconds](float DeltaTime) mutable
@@ -1201,7 +1201,7 @@ void FFrontendFilter_InUseByLoadedLevels::UnregisterDelayedRefresh()
 {
 	if (DelayedRefreshHandle.IsValid())
 	{
-		FTicker::GetCoreTicker().RemoveTicker(DelayedRefreshHandle);
+		FTSTicker::GetCoreTicker().RemoveTicker(DelayedRefreshHandle);
 		DelayedRefreshHandle.Reset();
 	}
 }
diff --git a/Engine/Source/Editor/ContentBrowser/Private/FrontendFilters.h b/Engine/Source/Editor/ContentBrowser/Private/FrontendFilters.h
--- a/Engine/Source/Editor/ContentBrowser/Private/FrontendFilters.h
+++ b/Engine/Source/Editor/ContentBrowser/Private/FrontendFilters.h
@@ -8,6 +8,7 @@
 #include "IAssetTools.h"
 #include "Misc/TextFilterExpressionEvaluator.h"
 #include "FrontendFilterBase.h"
+#include "Containers/Ticker.h"
 
 #define LOCTEXT_NAMESPACE "ContentBrowser"
 
@@ -296,7 +297,7 @@ private:
 	void Refresh();
 	void RegisterDelayedRefresh(float DelayInSeconds);
 	void UnregisterDelayedRefresh();
-	FDelegateHandle DelayedRefreshHandle;
+	FTSTicker::FDelegateHandle DelayedRefreshHandle;
 	bool bIsDirty = false;
 	bool bIsCurrentlyActive = false;
 };
diff --git a/Engine/Source/Editor/ContentBrowserData/Private/ContentBrowserDataSubsystem.cpp b/Engine/Source/Editor/ContentBrowserData/Private/ContentBrowserDataSubsystem.cpp
--- a/Engine/Source/Editor/ContentBrowserData/Private/ContentBrowserDataSubsystem.cpp
+++ b/Engine/Source/Editor/ContentBrowserData/Private/ContentBrowserDataSubsystem.cpp
@@ -43,7 +43,7 @@ void UContentBrowserDataSubsystem::Initialize(FSubsystemCollectionBase& Collecti
 	FEditorDelegates::EndPIE.AddUObject(this, &UContentBrowserDataSubsystem::OnEndPIE);
 
 	// Tick during normal operation
-	TickHandle = FTicker::GetCoreTicker().AddTicker(TEXT("ContentBrowserData"), 0.1f, [this](const float InDeltaTime)
+	TickHandle = FTSTicker::GetCoreTicker().AddTicker(TEXT("ContentBrowserData"), 0.1f, [this](const float InDeltaTime)
 	{
 		Tick(InDeltaTime);
 		return true;
@@ -71,7 +71,7 @@ void UContentBrowserDataSubsystem::Deinitialize()
 
 	if (TickHandle.IsValid())
 	{
-		FTicker::GetCoreTicker().RemoveTicker(TickHandle);
+		FTSTicker::GetCoreTicker().RemoveTicker(TickHandle);
 		TickHandle.Reset();
 	}
 
diff --git a/Engine/Source/Editor/PluginWarden/Private/PluginWardenModule.cpp b/Engine/Source/Editor/PluginWarden/Private/PluginWardenModule.cpp
--- a/Engine/Source/Editor/PluginWarden/Private/PluginWardenModule.cpp
+++ b/Engine/Source/Editor/PluginWarden/Private/PluginWardenModule.cpp
@@ -99,7 +99,7 @@ bool FPluginWardenModule::RunAuthorizationPipeline(const FText& PluginFriendlyNa
 		LastTickTime = CurrentTickTime;
 
 		FTaskGraphInterface::Get().ProcessThreadUntilIdle(ENamedThreads::GameThread);
-		FTicker::GetCoreTicker().Tick(DeltaTime);
+		FTSTicker::GetCoreTicker().Tick(DeltaTime);
 
 		PreviousState = AuthorizationState;
 		AuthorizationState = Authorizer.UpdateAuthorizationState(DeltaTime);
diff --git a/Engine/Source/Editor/PluginWarden/Private/SAuthorizingPlugin.cpp b/Engine/Source/Editor/PluginWarden/Private/SAuthorizingPlugin.cpp
--- a/Engine/Source/Editor/PluginWarden/Private/SAuthorizingPlugin.cpp
+++ b/Engine/Source/Editor/PluginWarden/Private/SAuthorizingPlugin.cpp
@@ -123,7 +123,7 @@ EActiveTimerReturnType SAuthorizingPlugin::RefreshStatus(double InCurrentTime, f
 	// Engine tick isn't running when the modal window is open, so we need to tick any core tickers
 	// to as that's what the RPC system uses to update the current state of RPC calls.
 	FTaskGraphInterface::Get().ProcessThreadUntilIdle(ENamedThreads::GameThread);
-	FTicker::GetCoreTicker().Tick(InDeltaTime);
+	FTSTicker::GetCoreTicker().Tick(InDeltaTime);
 
 	PreviousAuthorizationState = AuthorizationState;
 	AuthorizationState = Authorizer->UpdateAuthorizationState(InDeltaTime);
diff --git a/Engine/Source/Editor/StatusBar/Private/NotificationTesting.cpp b/Engine/Source/Editor/StatusBar/Private/NotificationTesting.cpp
--- a/Engine/Source/Editor/StatusBar/Private/NotificationTesting.cpp
+++ b/Engine/Source/Editor/StatusBar/Private/NotificationTesting.cpp
@@ -68,7 +68,7 @@ static void TestNotifications()
 {
 	static const float Timeout = 15.0f;
 
-	FTicker::GetCoreTicker().AddTicker(TEXT("TestNotifications"), 0.0f, [](float DeltaTime)
+	FTSTicker::GetCoreTicker().AddTicker(TEXT("TestNotifications"), 0.0f, [](float DeltaTime)
 		{
 			FNotificationInfo NotificationInfo(FText::FromString(TEXT("Test Notification 1")));
 			NotificationInfo.FadeInDuration = 2.0f;
@@ -81,7 +81,7 @@ static void TestNotifications()
 			return false;
 		});
 
-	FTicker::GetCoreTicker().AddTicker(TEXT("TestNotifications"), 2.0f, [](float DeltaTime)
+	FTSTicker::GetCoreTicker().AddTicker(TEXT("TestNotifications"), 2.0f, [](float DeltaTime)
 		{
 			FNotificationInfo NotificationInfo(FText::FromString(TEXT("Test Notification 2")));
 			NotificationInfo.FadeInDuration = 2.0f;
@@ -96,7 +96,7 @@ static void TestNotifications()
 			return false;
 		});
 
-	FTicker::GetCoreTicker().AddTicker(TEXT("TestNotifications"), 3.0f, [](float DeltaTime)
+	FTSTicker::GetCoreTicker().AddTicker(TEXT("TestNotifications"), 3.0f, [](float DeltaTime)
 		{
 			FNotificationInfo NotificationInfo(FText::FromString(TEXT("Test Notification 3")));
 			NotificationInfo.FadeInDuration = 2.0f;
@@ -111,7 +111,7 @@ static void TestNotifications()
 			return false;
 		});
 
-	FTicker::GetCoreTicker().AddTicker(TEXT("TestNotifications"), 4.0f, [](float DeltaTime)
+	FTSTicker::GetCoreTicker().AddTicker(TEXT("TestNotifications"), 4.0f, [](float DeltaTime)
 		{
 			FNotificationInfo NotificationInfo(FText::FromString(TEXT("Test Notification 4")));
 			NotificationInfo.FadeInDuration = 2.0f;
@@ -126,7 +126,7 @@ static void TestNotifications()
 			return false;
 		});
 
-	FTicker::GetCoreTicker().AddTicker(TEXT("TestNotifications"), 5.0f, [](float DeltaTime)
+	FTSTicker::GetCoreTicker().AddTicker(TEXT("TestNotifications"), 5.0f, [](float DeltaTime)
 		{
 			FNotificationInfo NotificationInfo(FText::FromString(TEXT("Test Notification 5")));
 			NotificationInfo.FadeInDuration = 2.0f;
@@ -144,7 +144,7 @@ static void TestNotifications()
 			return false;
 		});
 
-	FTicker::GetCoreTicker().AddTicker(TEXT("TestNotifications"), 6.0f, [](float DeltaTime)
+	FTSTicker::GetCoreTicker().AddTicker(TEXT("TestNotifications"), 6.0f, [](float DeltaTime)
 		{
 			FNotificationInfo NotificationInfo(FText::FromString(TEXT("Test Notification 6")));
 			NotificationInfo.FadeInDuration = 2.0f;
@@ -161,7 +161,7 @@ static void TestNotifications()
 			return false;
 		});
 
-	FTicker::GetCoreTicker().AddTicker(TEXT("TestNotifications"), 7.0f, [](float DeltaTime)
+	FTSTicker::GetCoreTicker().AddTicker(TEXT("TestNotifications"), 7.0f, [](float DeltaTime)
 		{
 			FNotificationInfo NotificationInfo(FText::FromString(TEXT("Test Notification 7")));
 			NotificationInfo.FadeInDuration = 2.0f;
@@ -180,7 +180,7 @@ static void TestNotifications()
 			return false;
 		});
 
-	FTicker::GetCoreTicker().AddTicker(TEXT("TestNotifications"), 8.0f, [](float DeltaTime)
+	FTSTicker::GetCoreTicker().AddTicker(TEXT("TestNotifications"), 8.0f, [](float DeltaTime)
 		{
 			FNotificationInfo NotificationInfo(FText::FromString(TEXT("Everything Under The Sun. This one also has a lot of text which should wrap to the next line")));
 			NotificationInfo.FadeInDuration = 2.0f;
diff --git a/Engine/Source/Editor/UnrealEd/Private/Commandlets/ContentCommandlets.cpp b/Engine/Source/Editor/UnrealEd/Private/Commandlets/ContentCommandlets.cpp
--- a/Engine/Source/Editor/UnrealEd/Private/Commandlets/ContentCommandlets.cpp
+++ b/Engine/Source/Editor/UnrealEd/Private/Commandlets/ContentCommandlets.cpp
@@ -1896,7 +1896,7 @@ void UResavePackagesCommandlet::PerformAdditionalOperations(class UWorld* World,
 
 				while (Processor->IsProxyGenerationRunning())
 				{
-					FTicker::GetCoreTicker().Tick(FApp::GetDeltaTime());
+					FTSTicker::GetCoreTicker().Tick(FApp::GetDeltaTime());
 					FThreadManager::Get().Tick();
 					FTaskGraphInterface::Get().ProcessThreadUntilIdle(ENamedThreads::GameThread);
 					FPlatformProcess::Sleep(0.1f);
diff --git a/Engine/Source/Editor/UnrealEd/Private/EditorEngine.cpp b/Engine/Source/Editor/UnrealEd/Private/EditorEngine.cpp
--- a/Engine/Source/Editor/UnrealEd/Private/EditorEngine.cpp
+++ b/Engine/Source/Editor/UnrealEd/Private/EditorEngine.cpp
@@ -1548,7 +1548,7 @@ void UEditorEngine::Tick( float DeltaSeconds, bool bIdleMode )
 	SourceCodeAccessModule.GetAccessor().Tick(DeltaSeconds);
 
 	// tick the directory watcher
-	// @todo: Put me into an FTicker that is created when the DW module is loaded
+	// @todo: Put me into an FTSTicker that is created when the DW module is loaded
 	if( !FApp::IsProjectNameEmpty() )
 	{
 		static FName DirectoryWatcherName("DirectoryWatcher");
diff --git a/Engine/Source/Editor/UnrealEd/Private/Elements/SMInstance/SMInstanceElementDetailsProxyObject.cpp b/Engine/Source/Editor/UnrealEd/Private/Elements/SMInstance/SMInstanceElementDetailsProxyObject.cpp
--- a/Engine/Source/Editor/UnrealEd/Private/Elements/SMInstance/SMInstanceElementDetailsProxyObject.cpp
+++ b/Engine/Source/Editor/UnrealEd/Private/Elements/SMInstance/SMInstanceElementDetailsProxyObject.cpp
@@ -16,7 +16,7 @@ void USMInstanceElementDetailsProxyObject::Initialize(const FSMInstanceElementId
 	ISMComponent = InSMInstanceElementId.ISMComponent;
 	ISMInstanceId = InSMInstanceElementId.InstanceId;
 
-	TickHandle = FTicker::GetCoreTicker().AddTicker(TEXT("USMInstanceElementDetailsProxyObject"), 0.1f, [this](float)
+	TickHandle = FTSTicker::GetCoreTicker().AddTicker(TEXT("USMInstanceElementDetailsProxyObject"), 0.1f, [this](float)
 	{
 		SyncProxyStateFromInstance();
 		return true;
@@ -28,7 +28,7 @@ void USMInstanceElementDetailsProxyObject::Shutdown()
 {
 	if (TickHandle.IsValid())
 	{
-		FTicker::GetCoreTicker().RemoveTicker(TickHandle);
+		FTSTicker::GetCoreTicker().RemoveTicker(TickHandle);
 		TickHandle.Reset();
 	}
 
diff --git a/Engine/Source/Editor/UnrealEd/Private/Elements/SMInstance/SMInstanceElementDetailsProxyObject.h b/Engine/Source/Editor/UnrealEd/Private/Elements/SMInstance/SMInstanceElementDetailsProxyObject.h
--- a/Engine/Source/Editor/UnrealEd/Private/Elements/SMInstance/SMInstanceElementDetailsProxyObject.h
+++ b/Engine/Source/Editor/UnrealEd/Private/Elements/SMInstance/SMInstanceElementDetailsProxyObject.h
@@ -7,6 +7,8 @@
 
 #include "IObjectNameEditSink.h"
 
+#include "Containers/Ticker.h"
+
 #include "SMInstanceElementDetailsProxyObject.generated.h"
 
 UCLASS(Transient)
@@ -32,7 +34,7 @@ private:
 	TWeakObjectPtr<UInstancedStaticMeshComponent> ISMComponent;
 	uint64 ISMInstanceId = 0;
 
-	FDelegateHandle TickHandle;
+	FTSTicker::FDelegateHandle TickHandle;
 	bool bIsWithinInteractiveTransformEdit = false;
 
 	friend class FSMInstanceElementDetailsProxyObjectNameEditSink;
diff --git a/Engine/Source/Editor/UnrealEd/Private/Subsystems/AssetEditorSubsystem.cpp b/Engine/Source/Editor/UnrealEd/Private/Subsystems/AssetEditorSubsystem.cpp
--- a/Engine/Source/Editor/UnrealEd/Private/Subsystems/AssetEditorSubsystem.cpp
+++ b/Engine/Source/Editor/UnrealEd/Private/Subsystems/AssetEditorSubsystem.cpp
@@ -48,7 +48,7 @@ UAssetEditorSubsystem::UAssetEditorSubsystem()
 void UAssetEditorSubsystem::Initialize(FSubsystemCollectionBase& Collection)
 {
 	TickDelegate = FTickerDelegate::CreateUObject(this, &UAssetEditorSubsystem::HandleTicker);
-	FTicker::GetCoreTicker().AddTicker(TickDelegate, 1.f);
+	FTSTicker::GetCoreTicker().AddTicker(TickDelegate, 1.f);
 
 	FCoreUObjectDelegates::OnPackageReloaded.AddUObject(this, &UAssetEditorSubsystem::HandlePackageReloaded);
 
diff --git a/Engine/Source/Programs/CrashReportClient/Private/CrashReportClient.cpp b/Engine/Source/Programs/CrashReportClient/Private/CrashReportClient.cpp
--- a/Engine/Source/Programs/CrashReportClient/Private/CrashReportClient.cpp
+++ b/Engine/Source/Programs/CrashReportClient/Private/CrashReportClient.cpp
@@ -81,7 +81,7 @@ FCrashReportClient::~FCrashReportClient()
 {
 	if (TickHandle.IsValid())
 	{
-		FTicker::GetCoreTicker().RemoveTicker(TickHandle);
+		FTSTicker::GetCoreTicker().RemoveTicker(TickHandle);
 		TickHandle.Reset();
 	}
 	StopBackgroundThread();
@@ -237,7 +237,7 @@ void FCrashReportClient::StartTicker()
 {
 	if (!TickHandle.IsValid())
 	{
-		TickHandle = FTicker::GetCoreTicker().AddTicker(FTickerDelegate::CreateRaw(this, &FCrashReportClient::Tick), 1.f);
+		TickHandle = FTSTicker::GetCoreTicker().AddTicker(FTickerDelegate::CreateRaw(this, &FCrashReportClient::Tick), 1.f);
 	}
 }
 
diff --git a/Engine/Source/Programs/CrashReportClient/Private/MainLoopTiming.cpp b/Engine/Source/Programs/CrashReportClient/Private/MainLoopTiming.cpp
--- a/Engine/Source/Programs/CrashReportClient/Private/MainLoopTiming.cpp
+++ b/Engine/Source/Programs/CrashReportClient/Private/MainLoopTiming.cpp
@@ -25,7 +25,6 @@ void FMainLoopTiming::Tick()
 
 	// Tick app logic
 	FTaskGraphInterface::Get().ProcessThreadUntilIdle(ENamedThreads::GameThread);
-	FTicker::GetCoreTicker().Tick(ActualDeltaTime);
 	FTSTicker::GetCoreTicker().Tick(ActualDeltaTime);
 
 #if !CRASH_REPORT_UNATTENDED_ONLY
diff --git a/Engine/Source/Programs/UnrealFileServer/Private/UnrealFileServer.cpp b/Engine/Source/Programs/UnrealFileServer/Private/UnrealFileServer.cpp
--- a/Engine/Source/Programs/UnrealFileServer/Private/UnrealFileServer.cpp
+++ b/Engine/Source/Programs/UnrealFileServer/Private/UnrealFileServer.cpp
@@ -52,7 +52,7 @@ INT32_MAIN_INT32_ARGC_TCHAR_ARGV()
 		float DeltaSeconds = LastTime - Now;
 		LastTime = Now;
 
-		// @todo: Put me into an FTicker that is created when the DW module is loaded
+		// @todo: Put me into an FTSTicker that is created when the DW module is loaded
 		FDirectoryWatcherModule& DirectoryWatcherModule = FModuleManager::Get().LoadModuleChecked<FDirectoryWatcherModule>(TEXT("DirectoryWatcher"));
 		DirectoryWatcherModule.Get()->Tick(Now - LastTime);
 
diff --git a/Engine/Source/Programs/UnrealRecoverySvc/Private/UnrealRecoverySvc.cpp b/Engine/Source/Programs/UnrealRecoverySvc/Private/UnrealRecoverySvc.cpp
--- a/Engine/Source/Programs/UnrealRecoverySvc/Private/UnrealRecoverySvc.cpp
+++ b/Engine/Source/Programs/UnrealRecoverySvc/Private/UnrealRecoverySvc.cpp
@@ -49,7 +49,7 @@ INT32_MAIN_INT32_ARGC_TCHAR_ARGV()
 		return ServerConfig;
 	};
 
-	FTicker::GetCoreTicker().AddTicker(TEXT("CheckEditorHealth"), 1.0f, [&EditorProcessId](float)
+	FTSTicker::GetCoreTicker().AddTicker(TEXT("CheckEditorHealth"), 1.0f, [&EditorProcessId](float)
 	{
 		if (!FPlatformProcess::IsApplicationRunning(EditorProcessId))
 		{
diff --git a/Engine/Source/Runtime/Analytics/AnalyticsET/Private/AnalyticsPerfTracker.h b/Engine/Source/Runtime/Analytics/AnalyticsET/Private/AnalyticsPerfTracker.h
--- a/Engine/Source/Runtime/Analytics/AnalyticsET/Private/AnalyticsPerfTracker.h
+++ b/Engine/Source/Runtime/Analytics/AnalyticsET/Private/AnalyticsPerfTracker.h
@@ -16,7 +16,7 @@
 #include "HAL/PlatformTime.h"
 
 /** Measures analytics bandwidth. Only active when -AnalyticsTrackPerf is specified on the command line. */
-struct FAnalyticsPerfTracker : FTickerObjectBase
+struct FAnalyticsPerfTracker : FTSTickerObjectBase
 {
 	FAnalyticsPerfTracker();
 
diff --git a/Engine/Source/Runtime/Analytics/AnalyticsET/Private/IAnalyticsProviderET.cpp b/Engine/Source/Runtime/Analytics/AnalyticsET/Private/IAnalyticsProviderET.cpp
--- a/Engine/Source/Runtime/Analytics/AnalyticsET/Private/IAnalyticsProviderET.cpp
+++ b/Engine/Source/Runtime/Analytics/AnalyticsET/Private/IAnalyticsProviderET.cpp
@@ -45,13 +45,13 @@ namespace AnalyticsProviderETCvars
  */
 class FAnalyticsProviderET :
 	public IAnalyticsProviderET,
-	public FTickerObjectBase,
+	public FTSTickerObjectBase,
 	public TSharedFromThis<FAnalyticsProviderET>
 {
 public:
 	FAnalyticsProviderET(const FAnalyticsET::Config& ConfigValues);
 
-	// FTickerObjectBase
+	// FTSTickerObjectBase
 
 	bool Tick(float DeltaSeconds) override;
 
diff --git a/Engine/Source/Runtime/Core/Private/Containers/BackgroundableTicker.cpp b/Engine/Source/Runtime/Core/Private/Containers/BackgroundableTicker.cpp
--- a/Engine/Source/Runtime/Core/Private/Containers/BackgroundableTicker.cpp
+++ b/Engine/Source/Runtime/Core/Private/Containers/BackgroundableTicker.cpp
@@ -38,6 +38,8 @@ FTSBackgroundableTicker::~FTSBackgroundableTicker()
 	FCoreDelegates::MobileBackgroundTickDelegate.Remove(BackgroundTickerHandle);
 }
 
+PRAGMA_DISABLE_DEPRECATION_WARNINGS
+
 //////////////////////////////////////////////////////////////////////////////////////////////////////////////////////////////////////
 // deprecated non thread-safe version
 
@@ -73,4 +75,6 @@ FBackgroundableTicker::~FBackgroundableTicker()
 {
 	FTicker::GetCoreTicker().RemoveTicker(CoreTickerHandle);
 	FCoreDelegates::MobileBackgroundTickDelegate.Remove(BackgroundTickerHandle);
-}
\ No newline at end of file
+}
+
+PRAGMA_ENABLE_DEPRECATION_WARNINGS
diff --git a/Engine/Source/Runtime/Core/Private/Containers/Ticker.cpp b/Engine/Source/Runtime/Core/Private/Containers/Ticker.cpp
--- a/Engine/Source/Runtime/Core/Private/Containers/Ticker.cpp
+++ b/Engine/Source/Runtime/Core/Private/Containers/Ticker.cpp
@@ -116,6 +116,15 @@ void FTSTicker::Tick(float DeltaTime)
 	}
 }
 
+void FTSTicker::Reset()
+{
+	while (AddedElements.Dequeue())
+	{
+	}
+	CurrentTime = 0.0;
+	Elements.Reset();
+}
+
 FTSTicker::FElement::FElement()
 	: FireTime(0.0)
 	, DelayTime(0.0f)
diff --git a/Engine/Source/Runtime/Core/Private/GenericPlatform/GenericPlatformCrashContext.cpp b/Engine/Source/Runtime/Core/Private/GenericPlatform/GenericPlatformCrashContext.cpp
--- a/Engine/Source/Runtime/Core/Private/GenericPlatform/GenericPlatformCrashContext.cpp
+++ b/Engine/Source/Runtime/Core/Private/GenericPlatform/GenericPlatformCrashContext.cpp
@@ -211,9 +211,9 @@ void FGenericCrashContext::Initialize()
 
 	// Initialize delegate for updating SecondsSinceStart, because FPlatformTime::Seconds() is not POSIX safe.
 	const float PollingInterval = 1.0f;
-	FTicker::GetCoreTicker().AddTicker( FTickerDelegate::CreateLambda( []( float DeltaTime )
+	FTSTicker::GetCoreTicker().AddTicker( FTickerDelegate::CreateLambda( []( float DeltaTime )
 	{
-    QUICK_SCOPE_CYCLE_COUNTER(STAT_NCachedCrashContextProperties_LambdaTicker);
+		QUICK_SCOPE_CYCLE_COUNTER(STAT_NCachedCrashContextProperties_LambdaTicker);
 
 		NCached::Session.SecondsSinceStart = int32(FPlatformTime::Seconds() - GStartTime);
 		return true;
diff --git a/Engine/Source/Runtime/Core/Private/GenericPlatform/GenericPlatformMemory.cpp b/Engine/Source/Runtime/Core/Private/GenericPlatform/GenericPlatformMemory.cpp
--- a/Engine/Source/Runtime/Core/Private/GenericPlatform/GenericPlatformMemory.cpp
+++ b/Engine/Source/Runtime/Core/Private/GenericPlatform/GenericPlatformMemory.cpp
@@ -144,7 +144,7 @@ void FGenericPlatformMemory::Init()
 #if	STATS
 	// Stats are updated only once per second.
 	const float PollingInterval = 1.0f;
-	FTicker::GetCoreTicker().AddTicker( FTickerDelegate::CreateStatic( &FGenericStatsUpdater::EnqueueUpdateStats ), PollingInterval );
+	FTSTicker::GetCoreTicker().AddTicker( FTickerDelegate::CreateStatic( &FGenericStatsUpdater::EnqueueUpdateStats ), PollingInterval );
 
 	// Update for the first time.
 	FGenericStatsUpdater::DoUpdateStats();
diff --git a/Engine/Source/Runtime/Core/Private/GenericPlatform/GenericPlatformTime.cpp b/Engine/Source/Runtime/Core/Private/GenericPlatform/GenericPlatformTime.cpp
--- a/Engine/Source/Runtime/Core/Private/GenericPlatform/GenericPlatformTime.cpp
+++ b/Engine/Source/Runtime/Core/Private/GenericPlatform/GenericPlatformTime.cpp
@@ -186,8 +186,8 @@ struct FCPUTimeDump
 			FParse::Value(*Args[0], TEXT("delay="), Delay);
 		}
 
-		FTicker::GetCoreTicker().RemoveTicker( GetCPUTimeDelegateHandle );
-		FTicker::GetCoreTicker().RemoveTicker( CPUTimeDumpDelegateHandle );
+		FTSTicker::GetCoreTicker().RemoveTicker( GetCPUTimeDelegateHandle );
+		FTSTicker::GetCoreTicker().RemoveTicker( CPUTimeDumpDelegateHandle );
 		GetCPUTimeDelegateHandle .Reset();
 		CPUTimeDumpDelegateHandle.Reset();
 
@@ -200,8 +200,8 @@ struct FCPUTimeDump
 			Delay = FMath::Clamp( Delay, 10, 300 );
 			UE_LOG(LogGenericPlatformTime, Log, TEXT("Delay set to %i second(s), started printing the CPU usage"), Delay);
 
-			GetCPUTimeDelegateHandle   = FTicker::GetCoreTicker().AddTicker( GetCPUTimeDelegate );
-			CPUTimeDumpDelegateHandle  = FTicker::GetCoreTicker().AddTicker( CPUTimeDumpDelegate, (float)Delay );
+			GetCPUTimeDelegateHandle   = FTSTicker::GetCoreTicker().AddTicker( GetCPUTimeDelegate );
+			CPUTimeDumpDelegateHandle  = FTSTicker::GetCoreTicker().AddTicker( CPUTimeDumpDelegate, (float)Delay );
 		}
 	}
 
@@ -244,10 +244,10 @@ protected:
 	FTickerDelegate CPUTimeDumpDelegate;
 
 	/** Handle for the added GetCPUTimeDelegate delegate. */
-	FDelegateHandle GetCPUTimeDelegateHandle;
+	FTSTicker::FDelegateHandle GetCPUTimeDelegateHandle;
 
 	/** Handle for the added CPUTimeDumpDelegate delegate. */
-	FDelegateHandle CPUTimeDumpDelegateHandle;
+	FTSTicker::FDelegateHandle CPUTimeDumpDelegateHandle;
 };
 
 static FAutoConsoleCommand CPUTimeDumpCommand
diff --git a/Engine/Source/Runtime/Core/Private/HoloLens/HoloLensTime.cpp b/Engine/Source/Runtime/Core/Private/HoloLens/HoloLensTime.cpp
--- a/Engine/Source/Runtime/Core/Private/HoloLens/HoloLensTime.cpp
+++ b/Engine/Source/Runtime/Core/Private/HoloLens/HoloLensTime.cpp
@@ -25,7 +25,7 @@ double FHoloLensTime::InitTiming()
 	static const float PollingInterval = 1.0f / 4.0f;
 
 	// Register a ticker delegate for updating the CPU utilization data.
-	FTicker::GetCoreTicker().AddTicker(FTickerDelegate::CreateStatic(&FPlatformTime::UpdateCPUTime), PollingInterval);
+	FTSTicker::GetCoreTicker().AddTicker(FTickerDelegate::CreateStatic(&FPlatformTime::UpdateCPUTime), PollingInterval);
 
 	return FPlatformTime::Seconds();
 }
diff --git a/Engine/Source/Runtime/Core/Private/IO/IoDispatcher.cpp b/Engine/Source/Runtime/Core/Private/IO/IoDispatcher.cpp
--- a/Engine/Source/Runtime/Core/Private/IO/IoDispatcher.cpp
+++ b/Engine/Source/Runtime/Core/Private/IO/IoDispatcher.cpp
@@ -156,14 +156,14 @@ public:
 			BatchAllocator.Trim();
 		});
 #if CSV_PROFILER
-		TickerHandle = FTicker::GetCoreTicker().AddTicker(FTickerDelegate::CreateRaw(this, &FIoDispatcherImpl::TickCsv));
+		TickerHandle = FTSTicker::GetCoreTicker().AddTicker(FTickerDelegate::CreateRaw(this, &FIoDispatcherImpl::TickCsv));
 #endif
 	}
 
 	~FIoDispatcherImpl()
 	{
 #if CSV_PROFILER
-		FTicker::GetCoreTicker().RemoveTicker(TickerHandle);
+		FTSTicker::GetCoreTicker().RemoveTicker(TickerHandle);
 #endif
 		delete Thread;
 		FPlatformProcess::ReturnSynchEventToPool(DispatcherEvent);
@@ -678,7 +678,7 @@ private:
 	uint64 PendingIoRequestsCount = 0;
 	int64 TotalLoaded = 0;
 #if CSV_PROFILER
-	FDelegateHandle TickerHandle;
+	FTSTicker::FDelegateHandle TickerHandle;
 #endif
 };
 
diff --git a/Engine/Source/Runtime/Core/Private/IO/IoDispatcherFileBackend.cpp b/Engine/Source/Runtime/Core/Private/IO/IoDispatcherFileBackend.cpp
--- a/Engine/Source/Runtime/Core/Private/IO/IoDispatcherFileBackend.cpp
+++ b/Engine/Source/Runtime/Core/Private/IO/IoDispatcherFileBackend.cpp
@@ -1912,20 +1912,20 @@ uint64 FFileIoStats::LastOffset = 0;
 uint32 FFileIoStats::FileIoStoreThreadId = 0;
 uint32 FFileIoStats::IoDispatcherThreadId = 0;
 
-FDelegateHandle FFileIoStats::TickerHandle;
+FTSTicker::FDelegateHandle FFileIoStats::TickerHandle;
 #endif
 
 void FFileIoStats::Init()
 {
 #if CSV_PROFILER
-	TickerHandle = FTicker::GetCoreTicker().AddTicker(FTickerDelegate::CreateStatic(&FFileIoStats::CsvTick));
+	TickerHandle = FTSTicker::GetCoreTicker().AddTicker(FTickerDelegate::CreateStatic(&FFileIoStats::CsvTick));
 #endif
 }
 
 void FFileIoStats::Shutdown()
 {
 #if CSV_PROFILER
-	FTicker::GetCoreTicker().RemoveTicker(TickerHandle);
+	FTSTicker::GetCoreTicker().RemoveTicker(TickerHandle);
 #endif
 }
 
diff --git a/Engine/Source/Runtime/Core/Private/IO/IoDispatcherFileBackendTypes.h b/Engine/Source/Runtime/Core/Private/IO/IoDispatcherFileBackendTypes.h
--- a/Engine/Source/Runtime/Core/Private/IO/IoDispatcherFileBackendTypes.h
+++ b/Engine/Source/Runtime/Core/Private/IO/IoDispatcherFileBackendTypes.h
@@ -6,6 +6,7 @@
 #include "IO/IoDispatcherBackend.h"
 #include "ProfilingDebugging/CsvProfiler.h"
 #include "ProfilingDebugging/CountersTrace.h"
+#include "Containers/Ticker.h"
 
 #define IO_DISPATCHER_FILE_STATS (COUNTERSTRACE_ENABLED || CSV_PROFILER)
 
@@ -861,7 +862,7 @@ struct FFileIoStats
 	static uint32 FileIoStoreThreadId;
 	static uint32 IoDispatcherThreadId;
 
-	static FDelegateHandle TickerHandle;
+	static FTSTicker::FDelegateHandle TickerHandle;
 #endif
 
 	static float BytesToApproxMB(uint64 Bytes) {
diff --git a/Engine/Source/Runtime/Core/Private/Misc/EmbeddedCommunication.cpp b/Engine/Source/Runtime/Core/Private/Misc/EmbeddedCommunication.cpp
--- a/Engine/Source/Runtime/Core/Private/Misc/EmbeddedCommunication.cpp
+++ b/Engine/Source/Runtime/Core/Private/Misc/EmbeddedCommunication.cpp
@@ -86,7 +86,7 @@ void FEmbeddedCommunication::Init()
  	GSleepEvent = FPlatformProcess::GetSynchEventFromPool(false);
 	GTickWithoutSleepCount = 0;
 
-	FTicker::GetCoreTicker().AddTicker(FTickerDelegate::CreateStatic(&FEmbeddedCommunication::TickGameThread));
+	FTSTicker::GetCoreTicker().AddTicker(FTickerDelegate::CreateStatic(&FEmbeddedCommunication::TickGameThread));
 #endif
 }
 
@@ -139,7 +139,7 @@ void FEmbeddedCommunication::ForceTick(int ID, float MinTimeSlice, float MaxTime
 	{
 		UE_LOG(LogInit, Display, TEXT("###ForceTick %d: processing messages..."), ID);
 		//We have to manually tick everything as we are looping the main thread here
-		FTicker::GetCoreTicker().Tick(Now - LastTime);
+		FTSTicker::GetCoreTicker().Tick(Now - LastTime);
 		FThreadManager::Get().Tick();
 		
 		FPlatformProcess::Sleep(DeltaTime);
diff --git a/Engine/Source/Runtime/Core/Private/Stats/Stats2.cpp b/Engine/Source/Runtime/Core/Private/Stats/Stats2.cpp
--- a/Engine/Source/Runtime/Core/Private/Stats/Stats2.cpp
+++ b/Engine/Source/Runtime/Core/Private/Stats/Stats2.cpp
@@ -241,7 +241,7 @@ void FStats::TickCommandletStats()
 		//check( ThreadStats->ScopeCount == 0 && TEXT( "FStats::TickCommandletStats must be called outside any scope counters" ) );
 
 		FTaskGraphInterface::Get().ProcessThreadUntilIdle( ENamedThreads::GameThread );
-		FTicker::GetCoreTicker().Tick( 1 / 60.0f );
+		FTSTicker::GetCoreTicker().Tick( 1 / 60.0f );
 
 		FStats::AdvanceFrame( false );
 	}
diff --git a/Engine/Source/Runtime/Core/Private/Unix/UnixPlatformCrashContext.cpp b/Engine/Source/Runtime/Core/Private/Unix/UnixPlatformCrashContext.cpp
--- a/Engine/Source/Runtime/Core/Private/Unix/UnixPlatformCrashContext.cpp
+++ b/Engine/Source/Runtime/Core/Private/Unix/UnixPlatformCrashContext.cpp
@@ -423,7 +423,7 @@ namespace UnixCrashReporterTracker
 		UnixCrashReporterTracker::MaxProcessSlots = ActiveProcessSlots;
 
         // Register our Tick function
-		FTicker::GetCoreTicker().AddTicker(FTickerDelegate::CreateStatic(&UnixCrashReporterTracker::Tick), 1.f);
+		FTSTicker::GetCoreTicker().AddTicker(FTickerDelegate::CreateStatic(&UnixCrashReporterTracker::Tick), 1.f);
 	}
 
 	/**
diff --git a/Engine/Source/Runtime/Core/Private/Windows/WindowsPlatformTime.cpp b/Engine/Source/Runtime/Core/Private/Windows/WindowsPlatformTime.cpp
--- a/Engine/Source/Runtime/Core/Private/Windows/WindowsPlatformTime.cpp
+++ b/Engine/Source/Runtime/Core/Private/Windows/WindowsPlatformTime.cpp
@@ -21,7 +21,7 @@ double FWindowsPlatformTime::InitTiming(void)
 	static const float PollingInterval = 1.0f / 4.0f;
 
 	// Register a ticker delegate for updating the CPU utilization data.
-	FTicker::GetCoreTicker().AddTicker( FTickerDelegate::CreateStatic( &FPlatformTime::UpdateCPUTime ), PollingInterval );
+	FTSTicker::GetCoreTicker().AddTicker( FTickerDelegate::CreateStatic( &FPlatformTime::UpdateCPUTime ), PollingInterval );
 
 	return FPlatformTime::Seconds();
 }
diff --git a/Engine/Source/Runtime/Core/Public/Containers/BackgroundableTicker.h b/Engine/Source/Runtime/Core/Public/Containers/BackgroundableTicker.h
--- a/Engine/Source/Runtime/Core/Public/Containers/BackgroundableTicker.h
+++ b/Engine/Source/Runtime/Core/Public/Containers/BackgroundableTicker.h
@@ -23,6 +23,8 @@ private:
 	bool bWasBackgrounded = false;
 };
 
+PRAGMA_DISABLE_DEPRECATION_WARNINGS
+
 /**
  * DEPRECATED not thread-safe version
  * For migration guide see `FTSTicker`
@@ -30,7 +32,7 @@ private:
  * This works the same as the core FTicker, but on supported mobile platforms 
  * it continues ticking while the app is running in the background.
  */
-class FBackgroundableTicker
+class UE_DEPRECATED(5.0, "Please use thread-safe FTSTickerObjectBase, see the migration guide in the code comment above FTSTicker") FBackgroundableTicker
 	: public FTicker
 {
 public:
@@ -45,4 +47,6 @@ private:
 	FDelegateHandle CoreTickerHandle;
 	FDelegateHandle BackgroundTickerHandle;
 	bool bWasBackgrounded = false;
-};
\ No newline at end of file
+};
+
+PRAGMA_ENABLE_DEPRECATION_WARNINGS
diff --git a/Engine/Source/Runtime/Core/Public/Containers/Ticker.h b/Engine/Source/Runtime/Core/Public/Containers/Ticker.h
--- a/Engine/Source/Runtime/Core/Public/Containers/Ticker.h
+++ b/Engine/Source/Runtime/Core/Public/Containers/Ticker.h
@@ -79,6 +79,11 @@ public:
 	 */
 	void Tick(float DeltaTime);
 
+	/**
+	 * Resets the instance to its default state. Must be called from the ticking thread.
+	 */
+	void Reset();
+
 private:
 	/** Internal structure to store a ticker delegate and related data **/
 	struct FElement
@@ -162,7 +167,7 @@ private:
  * `FTicker` to `FTSTicker` and change the type of the delegate handle returned by `AddTicker` methods from `FDelegateHandle` to
  * `FTSTicker::DelegateHandle`
  */
-class FTicker
+class UE_DEPRECATED(5.0, "Please use thread-safe FTSTicker, see the migration guide in the code comment above") FTicker
 {
 public:
 
@@ -249,11 +254,10 @@ protected:
 	bool bCurrentElementRemoved;
 };
 
-
 /**
  * Base class for ticker objects
  */
-class FTickerObjectBase
+class UE_DEPRECATED(5.0, "Please use thread-safe FTSTickerObjectBase, see the migration guide in the code comment above FTSTicker") FTickerObjectBase
 {
 	PRAGMA_DISABLE_DEPRECATION_WARNINGS // FTicker
 
diff --git a/Engine/Source/Runtime/Core/Public/Misc/EmbeddedCommunication.h b/Engine/Source/Runtime/Core/Public/Misc/EmbeddedCommunication.h
--- a/Engine/Source/Runtime/Core/Public/Misc/EmbeddedCommunication.h
+++ b/Engine/Source/Runtime/Core/Public/Misc/EmbeddedCommunication.h
@@ -40,7 +40,7 @@ public:
 	// returns true if NativeToEmbedded delegate for subsystem exists
 	static bool IsEmbeddedSubsystemAvailable(FName SubsystemName);
 	
-	// FTicker-like delegate, to bind things to be ticked at a regular interval while the game thread is otherwise asleep.
+	// FTSTicker-like delegate, to bind things to be ticked at a regular interval while the game thread is otherwise asleep.
 	static FSimpleMulticastDelegate SleepTickDelegate;
 
 	// get/set an object by name, thread safe
diff --git a/Engine/Source/Runtime/CoreUObject/Private/Serialization/AsyncLoadingFlushContext.cpp b/Engine/Source/Runtime/CoreUObject/Private/Serialization/AsyncLoadingFlushContext.cpp
--- a/Engine/Source/Runtime/CoreUObject/Private/Serialization/AsyncLoadingFlushContext.cpp
+++ b/Engine/Source/Runtime/CoreUObject/Private/Serialization/AsyncLoadingFlushContext.cpp
@@ -41,10 +41,10 @@ void FAsyncLoadingFlushContext::Flush(const FOnAsyncLoadingFlushComplete& OnFlus
 	OnFlushCompleteDelegate = OnFlushComplete;
 
 	// Check each frame whether async loading has completed.
-	LoadingCompleteCheckDelegateHandle = FTicker::GetCoreTicker().AddTicker(FTickerDelegate::CreateRaw(this, &FAsyncLoadingFlushContext::OnAsyncLoadingCheck));
+	LoadingCompleteCheckDelegateHandle = FTSTicker::GetCoreTicker().AddTicker(FTickerDelegate::CreateRaw(this, &FAsyncLoadingFlushContext::OnAsyncLoadingCheck));
 
 	// Periodically warn when waiting for async loading to complete.
-	WaitingWarnDelegateHandle = FTicker::GetCoreTicker().AddTicker(FTickerDelegate::CreateRaw(this, &FAsyncLoadingFlushContext::OnAsyncLoadingWarn), LoadingFlushWarnIntervalSeconds);
+	WaitingWarnDelegateHandle = FTSTicker::GetCoreTicker().AddTicker(FTickerDelegate::CreateRaw(this, &FAsyncLoadingFlushContext::OnAsyncLoadingWarn), LoadingFlushWarnIntervalSeconds);
 
 	StartTime = FPlatformTime::Seconds();
 	UE_LOG(AsyncLoadingFlush, Verbose, TEXT("Context[%s:%d] Flushing async loading."), *Context, Id);
@@ -54,13 +54,13 @@ void FAsyncLoadingFlushContext::CleanupTickers()
 {
 	if (LoadingCompleteCheckDelegateHandle.IsValid())
 	{
-		FTicker::GetCoreTicker().RemoveTicker(LoadingCompleteCheckDelegateHandle);
+		FTSTicker::GetCoreTicker().RemoveTicker(LoadingCompleteCheckDelegateHandle);
 		LoadingCompleteCheckDelegateHandle.Reset();
 	}
 
 	if (WaitingWarnDelegateHandle.IsValid())
 	{
-		FTicker::GetCoreTicker().RemoveTicker(WaitingWarnDelegateHandle);
+		FTSTicker::GetCoreTicker().RemoveTicker(WaitingWarnDelegateHandle);
 		WaitingWarnDelegateHandle.Reset();
 	}
 }
diff --git a/Engine/Source/Runtime/CoreUObject/Public/Serialization/AsyncLoadingFlushContext.h b/Engine/Source/Runtime/CoreUObject/Public/Serialization/AsyncLoadingFlushContext.h
--- a/Engine/Source/Runtime/CoreUObject/Public/Serialization/AsyncLoadingFlushContext.h
+++ b/Engine/Source/Runtime/CoreUObject/Public/Serialization/AsyncLoadingFlushContext.h
@@ -4,6 +4,7 @@
 
 #include "Delegates/Delegate.h"
 #include "Logging/LogMacros.h"
+#include "Containers/Ticker.h"
 
 DECLARE_LOG_CATEGORY_EXTERN(AsyncLoadingFlush, Log, All);
 
@@ -47,8 +48,8 @@ private:
 	double StartTime = 0;
 
 	/** Tick delegates. */
-	FDelegateHandle LoadingCompleteCheckDelegateHandle;
-	FDelegateHandle WaitingWarnDelegateHandle;
+	FTSTicker::FDelegateHandle LoadingCompleteCheckDelegateHandle;
+	FTSTicker::FDelegateHandle WaitingWarnDelegateHandle;
 
 	/** Increment ids so each instance has a unique id. */
 	static int32 NextId;
diff --git a/Engine/Source/Runtime/CrashReportCore/Private/CrashUpload.cpp b/Engine/Source/Runtime/CrashReportCore/Private/CrashUpload.cpp
--- a/Engine/Source/Runtime/CrashReportCore/Private/CrashUpload.cpp
+++ b/Engine/Source/Runtime/CrashReportCore/Private/CrashUpload.cpp
@@ -671,7 +671,7 @@ void FCrashUploadToReceiver::SendPingRequest()
 
 	if (Request->ProcessRequest())
 	{
-		FTicker::GetCoreTicker().AddTicker(FTickerDelegate::CreateRaw(this, &FCrashUploadToReceiver::PingTimeout), CrashUploadDefs::PingTimeoutSeconds);
+		FTSTicker::GetCoreTicker().AddTicker(FTickerDelegate::CreateRaw(this, &FCrashUploadToReceiver::PingTimeout), CrashUploadDefs::PingTimeoutSeconds);
 	}
 	else
 	{
diff --git a/Engine/Source/Runtime/Engine/Classes/Engine/Engine.h b/Engine/Source/Runtime/Engine/Classes/Engine/Engine.h
--- a/Engine/Source/Runtime/Engine/Classes/Engine/Engine.h
+++ b/Engine/Source/Runtime/Engine/Classes/Engine/Engine.h
@@ -19,6 +19,7 @@
 #include "RHI.h"
 #include "AudioDeviceManager.h"
 #include "Templates/UniqueObj.h"
+#include "Containers/Ticker.h"
 #include "Engine.generated.h"
 
 #define WITH_DYNAMIC_RESOLUTION (!UE_SERVER)
@@ -3428,7 +3429,7 @@ private:
 		bool Tick(float Seconds);
 
 		TMap<uint32, uint32>	MessagesToCountMap;
-		FDelegateHandle			TickerHandle;
+		FTSTicker::FDelegateHandle			TickerHandle;
 		float					DisplayTime;
 	};
 
diff --git a/Engine/Source/Runtime/Engine/Private/Components/TextRenderComponent.cpp b/Engine/Source/Runtime/Engine/Private/Components/TextRenderComponent.cpp
--- a/Engine/Source/Runtime/Engine/Private/Components/TextRenderComponent.cpp
+++ b/Engine/Source/Runtime/Engine/Private/Components/TextRenderComponent.cpp
@@ -24,6 +24,7 @@
 #include "Engine/TextRenderActor.h"
 #include "Materials/MaterialInstanceDynamic.h"
 #include "StaticMeshResources.h"
+#include "Containers/Ticker.h"
 
 #define LOCTEXT_NAMESPACE "TextRenderComponent"
 
@@ -506,12 +507,12 @@ private:
 
 	FTextRenderComponentMIDCache()
 	{
-		TickerHandle = FTicker::GetCoreTicker().AddTicker(FTickerDelegate::CreateRaw(this, &FTextRenderComponentMIDCache::PurgeUnreferencedMIDsTicker), 10.0f);
+		TickerHandle = FTSTicker::GetCoreTicker().AddTicker(FTickerDelegate::CreateRaw(this, &FTextRenderComponentMIDCache::PurgeUnreferencedMIDsTicker), 10.0f);
 	}
 
 	~FTextRenderComponentMIDCache()
 	{
-		FTicker::GetCoreTicker().RemoveTicker(TickerHandle);
+		FTSTicker::GetCoreTicker().RemoveTicker(TickerHandle);
 	}
 
 	bool PurgeUnreferencedMIDsTicker(float)
@@ -560,7 +561,7 @@ private:
 		MIDsPendingPurge = MoveTemp(MIDsToPurgeLater);
 	}
 
-	FDelegateHandle TickerHandle;
+	FTSTicker::FDelegateHandle TickerHandle;
 
 	TMap<FKey, FMIDDataPtr> CachedMIDs;
 
diff --git a/Engine/Source/Runtime/Engine/Private/Net/NetworkGranularMemoryLogging.cpp b/Engine/Source/Runtime/Engine/Private/Net/NetworkGranularMemoryLogging.cpp
--- a/Engine/Source/Runtime/Engine/Private/Net/NetworkGranularMemoryLogging.cpp
+++ b/Engine/Source/Runtime/Engine/Private/Net/NetworkGranularMemoryLogging.cpp
@@ -10,6 +10,7 @@
 #include "Containers/Ticker.h"
 #include "Templates/Sorting.h"
 #include "Stats/Stats.h"
+#include "Containers/Ticker.h"
 
 namespace GranularNetworkMemoryTrackingPrivate
 {
@@ -160,7 +161,7 @@ namespace GranularNetworkMemoryTrackingPrivate
 	private:
 
 		FNetworkMemoryTrackingScopeStack() :
-			TickHandle(FTicker::GetCoreTicker().AddTicker(TEXT("NetworkGranularMemoryLogging::FNetworkMemoryTrackingScopeStack"), 0, &FNetworkMemoryTrackingScopeStack::OnTick))
+			TickHandle(FTSTicker::GetCoreTicker().AddTicker(TEXT("NetworkGranularMemoryLogging::FNetworkMemoryTrackingScopeStack"), 0, &FNetworkMemoryTrackingScopeStack::OnTick))
 		{
 		}
 
@@ -187,12 +188,12 @@ namespace GranularNetworkMemoryTrackingPrivate
 			// TODO: Replace \r\n with some platform specific newline macro.
 			UE_LOG(LogNet, Warning, TEXT("\r\n%s"), *FString::Join(Rows, TEXT("\r\n\r\n")));
 
-			FTicker::GetCoreTicker().RemoveTicker(ScopeStack->TickHandle);
+			FTSTicker::GetCoreTicker().RemoveTicker(ScopeStack->TickHandle);
 			ScopeStack.Reset();
 			return true;
 		}
 
-		FDelegateHandle TickHandle;
+		FTSTicker::FDelegateHandle TickHandle;
 		TMap<FString, TUniquePtr<FNetworkMemoryTrackingScope>> TopLevelScopes;
 		FNetworkMemoryTrackingScope* CurrentScope = nullptr;
 		TArray<FNetworkMemoryTrackingScope*> ScopeStack;
diff --git a/Engine/Source/Runtime/Engine/Private/ODSC/ODSCManager.cpp b/Engine/Source/Runtime/Engine/Private/ODSC/ODSCManager.cpp
--- a/Engine/Source/Runtime/Engine/Private/ODSC/ODSCManager.cpp
+++ b/Engine/Source/Runtime/Engine/Private/ODSC/ODSCManager.cpp
@@ -14,7 +14,7 @@ DEFINE_LOG_CATEGORY(LogODSC);
 FODSCManager* GODSCManager = nullptr;
 
 FODSCManager::FODSCManager()
-	: FTickerObjectBase(0.0f, FBackgroundableTicker::GetCoreTicker())
+	: FTSTickerObjectBase(0.0f, FTSBackgroundableTicker::GetCoreTicker())
 	, Thread(new FODSCThread())
 {
 	Thread->StartThread();
diff --git a/Engine/Source/Runtime/Engine/Private/ProfilingDebugging/MallocLeakReporter.cpp b/Engine/Source/Runtime/Engine/Private/ProfilingDebugging/MallocLeakReporter.cpp
--- a/Engine/Source/Runtime/Engine/Private/ProfilingDebugging/MallocLeakReporter.cpp
+++ b/Engine/Source/Runtime/Engine/Private/ProfilingDebugging/MallocLeakReporter.cpp
@@ -48,7 +48,7 @@ void FMallocLeakReporter::Start(int32 FilterSize/*=0*/, float ReportOnTime /*= 0
 	FMallocLeakDetection::Get().SetAllocationCollection(true, FilterSize);
 
 	// Create a ticker to issue checkpoints
-	CheckpointTicker = FTicker::GetCoreTicker().AddTicker(FTickerDelegate::CreateLambda([this](const float TimeDelta)
+	CheckpointTicker = FTSTicker::GetCoreTicker().AddTicker(FTickerDelegate::CreateLambda([this](const float TimeDelta)
 	{
 		Checkpoint();
 		return Enabled;
@@ -57,7 +57,7 @@ void FMallocLeakReporter::Start(int32 FilterSize/*=0*/, float ReportOnTime /*= 0
 	// If specified, create a handler to generate reports periodically
 	if (ReportOnTime > 0.0f)
 	{
-		ReportTicker = FTicker::GetCoreTicker().AddTicker(FTickerDelegate::CreateLambda([this](const float TimeDelta)
+		ReportTicker = FTSTicker::GetCoreTicker().AddTicker(FTickerDelegate::CreateLambda([this](const float TimeDelta)
 		{
 			WriteReports();
 			return true;
@@ -79,8 +79,8 @@ void FMallocLeakReporter::Stop()
 	UE_LOG(LogLeakDetector, Log, TEXT("Stopped tracking allocations"));
 	FMallocLeakDetection::Get().SetAllocationCollection(false);
 
-	FTicker::GetCoreTicker().RemoveTicker(CheckpointTicker);
-	FTicker::GetCoreTicker().RemoveTicker(ReportTicker);
+	FTSTicker::GetCoreTicker().RemoveTicker(CheckpointTicker);
+	FTSTicker::GetCoreTicker().RemoveTicker(ReportTicker);
 
 	CheckpointTicker.Reset();
 	ReportTicker.Reset();
