AutomationDriver - Fix sequence tests on Mac

diff --git a/Engine/Plugins/Tests/AutomationDriverTests/Source/AutomationDriverTests/Private/AutomationDriver.spec.cpp b/Engine/Plugins/Tests/AutomationDriverTests/Source/AutomationDriverTests/Private/AutomationDriver.spec.cpp
--- a/Engine/Plugins/Tests/AutomationDriverTests/Source/AutomationDriverTests/Private/AutomationDriver.spec.cpp
+++ b/Engine/Plugins/Tests/AutomationDriverTests/Source/AutomationDriverTests/Private/AutomationDriver.spec.cpp
@@ -1070,8 +1070,14 @@ void FAutomationDriverSpec::Define()
 			Sequence->Actions()
 				.Click(By::Id("KeyModifierA#"))
 				.Click(By::Id("KeyA#"))
+				// Workaround for Mac. The first click closes previously opened popup menu
+				// and the second click opens needed popup menu
+				.Click(By::Id("KeyModifierEb"))
 				.Click(By::Id("KeyModifierEb"))
 				.Click(By::Id("KeyEb"))
+				// Workaround for Mac. The first click closes previously opened popup menu
+				// and the second click opens needed popup menu
+				.Click(By::Id("KeyModifierB#"))
 				.Click(By::Id("KeyModifierB#"))
 				.Click(By::Id("KeyB#"));
 
@@ -1179,43 +1185,35 @@ void FAutomationDriverSpec::Define()
 					TEST_TRUE(ElementA->IsFocused());
 				});
 
-				It("should cancel further sequence execution after timing out", EAsyncExecution::ThreadPool, [this]()
+				Describe("should cancel further sequence execution after timing out", [this]()
 				{
-					FDriverElementRef ElementB = Driver->FindElement(By::Id("KeyB"));
-					ElementB->Focus();
+					It("", EAsyncExecution::ThreadPool, [this]()
+					{
+						FDriverElementRef ElementB = Driver->FindElement(By::Id("KeyB"));
+						ElementB->Focus();
 
-					TSharedRef<int32, ESPMode::ThreadSafe> OriginalExistenceChecker = MakeShareable(new int32);
-					TWeakPtr<int32, ESPMode::ThreadSafe> WeakExistenceChecker = OriginalExistenceChecker;
-					AsyncTask(
-						ENamedThreads::GameThread,
-						[this, WeakExistenceChecker]()
+						Driver->GetConfiguration()->ImplicitWait = FTimespan::FromSeconds(0.5);
+						FDriverElementRef ElementA = Driver->FindElement(By::Id("KeyA"));
+
+						Async(EAsyncExecution::TaskGraphMainThread, [this]()
 						{
-							const TSharedPtr<int32, ESPMode::ThreadSafe> ExistenceChecker = WeakExistenceChecker.Pin();
-							if (ExistenceChecker.IsValid())
-							{
-								SuiteWidget->RemoveContents();
-								FTSTicker::GetCoreTicker().AddTicker(FTickerDelegate::CreateLambda([this, WeakExistenceChecker](float Delta) {
-									const TSharedPtr<int32, ESPMode::ThreadSafe> TickerExistenceChecker = WeakExistenceChecker.Pin();
-									if (TickerExistenceChecker.IsValid())
-									{
-										SuiteWidget->RestoreContents();
-									}
-									return false;
-								}), 2);
-							}
-						}
-					);
+							// Note that we have to call this method in GameThread to avoid assertions
+							SuiteWidget->RemoveContents();
+						}).Get();
 
-					Driver->GetConfiguration()->ImplicitWait = FTimespan::FromSeconds(0.5);
-					FDriverElementRef ElementA = Driver->FindElement(By::Id("KeyA"));
+						FDriverSequenceRef Sequence = Driver->CreateSequence();
+						Sequence->Actions()
+							.Wait(Until::ElementExists(ElementA, FWaitInterval::InSeconds(0.25), FWaitTimeout::InSeconds(1)))
+							.Focus(ElementA);
 
-					FDriverSequenceRef Sequence = Driver->CreateSequence();
-					Sequence->Actions()
-						.Wait(Until::ElementExists(ElementA, FWaitInterval::InSeconds(0.25), FWaitTimeout::InSeconds(1)))
-						.Focus(ElementA);
+						TEST_FALSE(Sequence->Perform());
+						TEST_FALSE(ElementA->IsFocused());
+					});
 
-					TEST_FALSE(Sequence->Perform());
-					TEST_FALSE(ElementA->IsFocused());
+					AfterEach([this]()
+					{
+						SuiteWidget->RestoreContents();
+					});
 				});
 			});
 
@@ -1260,43 +1258,30 @@ void FAutomationDriverSpec::Define()
 					TEST_TRUE(ElementA->IsFocused());
 				});
 
-				It("should cancel further sequence execution after timing out", EAsyncExecution::ThreadPool, [this]()
+				Describe("should cancel further sequence execution after timing out", [this]()
 				{
-					FDriverElementRef ElementB = Driver->FindElement(By::Id("KeyB"));
-					ElementB->Focus();
-
-					TSharedRef<int32, ESPMode::ThreadSafe> OriginalExistenceChecker = MakeShareable(new int32);
-					TWeakPtr<int32, ESPMode::ThreadSafe> WeakExistenceChecker = OriginalExistenceChecker;
-					SuiteViewModel->SetPianoVisibility(EVisibility::Collapsed);
-					AsyncTask(
-						ENamedThreads::GameThread,
-						[this, WeakExistenceChecker]()
-						{
-							const TSharedPtr<int32, ESPMode::ThreadSafe> ExistenceChecker = WeakExistenceChecker.Pin();
-							if (ExistenceChecker.IsValid())
-							{
-								FTSTicker::GetCoreTicker().AddTicker(FTickerDelegate::CreateLambda([this, WeakExistenceChecker](float Delta) {
-									const TSharedPtr<int32, ESPMode::ThreadSafe> TickerExistenceChecker = WeakExistenceChecker.Pin();
-									if (TickerExistenceChecker.IsValid())
-									{
-										SuiteViewModel->SetPianoVisibility(EVisibility::Visible);
-									}
-									return false;
-								}), 2);
-							}
-						}
-					);
-
-					Driver->GetConfiguration()->ImplicitWait = FTimespan::FromSeconds(0.5);
-					FDriverElementRef ElementA = Driver->FindElement(By::Id("KeyA"));
-
-					FDriverSequenceRef Sequence = Driver->CreateSequence();
-					Sequence->Actions()
-						.Wait(Until::ElementIsVisible(ElementA, FWaitInterval::InSeconds(0.25), FWaitTimeout::InSeconds(1)))
-						.Focus(ElementA);
-
-					TEST_FALSE(Sequence->Perform());
-					TEST_FALSE(ElementA->IsFocused());
+					It("", EAsyncExecution::ThreadPool, [this]()
+					{
+						FDriverElementRef ElementB = Driver->FindElement(By::Id("KeyB"));
+						ElementB->Focus();
+
+						Driver->GetConfiguration()->ImplicitWait = FTimespan::FromSeconds(0.5);
+						FDriverElementRef ElementA = Driver->FindElement(By::Id("KeyA"));
+
+						SuiteViewModel->SetPianoVisibility(EVisibility::Collapsed);
+
+						FDriverSequenceRef Sequence = Driver->CreateSequence();
+						Sequence->Actions()
+							.Wait(Until::ElementIsVisible(ElementA, FWaitInterval::InSeconds(0.25), FWaitTimeout::InSeconds(1)))
+							.Focus(ElementA);
+
+						TEST_FALSE(Sequence->Perform());
+						TEST_FALSE(ElementA->IsFocused());
+					});
+				
+					AfterEach([this]() {
+						SuiteViewModel->SetPianoVisibility(EVisibility::Visible);
+					});
 				});
 			});
 
@@ -1340,42 +1325,32 @@ void FAutomationDriverSpec::Define()
 					TEST_TRUE(ElementB->IsFocused());
 				});
 
-				It("should cancel further sequence execution after timing out", EAsyncExecution::ThreadPool, [this]()
+				Describe("should cancel further sequence execution after timing out", [this]()
 				{
-					TSharedRef<int32, ESPMode::ThreadSafe> OriginalExistenceChecker = MakeShareable(new int32);
-					TWeakPtr<int32, ESPMode::ThreadSafe> WeakExistenceChecker = OriginalExistenceChecker;
-					SuiteViewModel->SetKeyResetDelay(FTimespan::FromSeconds(2));
-					AsyncTask(
-						ENamedThreads::GameThread,
-						[this, WeakExistenceChecker]()
-						{
-							const TSharedPtr<int32, ESPMode::ThreadSafe> ExistenceChecker = WeakExistenceChecker.Pin();
-							if (ExistenceChecker.IsValid())
-							{
-								FTSTicker::GetCoreTicker().AddTicker(FTickerDelegate::CreateLambda([this, WeakExistenceChecker](float Delta) {
-									const TSharedPtr<int32, ESPMode::ThreadSafe> TickerExistenceChecker = WeakExistenceChecker.Pin();
-									if (TickerExistenceChecker.IsValid())
-									{
-										SuiteViewModel->SetKeyResetDelay(FTimespan::Zero());
-									}
-									return false;
-								}), 2);
-							}
-						}
-					);
+					BeforeEach([this]() {
+						SuiteViewModel->SetKeyResetDelay(FTimespan::FromSeconds(5.0));
+					});
 
-					Driver->GetConfiguration()->ImplicitWait = FTimespan::FromSeconds(0.5);
-					FDriverElementRef ElementB = Driver->FindElement(By::Id("KeyB"));
-					FDriverElementRef ElementA = Driver->FindElement(By::Id("KeyA"));
-					ElementA->Click();
+					It("", EAsyncExecution::ThreadPool, [this]()
+					{
+						Driver->GetConfiguration()->ImplicitWait = FTimespan::FromSeconds(0.5);
+						FDriverElementRef ElementB = Driver->FindElement(By::Id("KeyB"));
+						FDriverElementRef ElementA = Driver->FindElement(By::Id("KeyA"));
 
-					FDriverSequenceRef Sequence = Driver->CreateSequence();
-					Sequence->Actions()
-						.Wait(Until::ElementIsInteractable(ElementA, FWaitInterval::InSeconds(0.25), FWaitTimeout::InSeconds(1)))
-						.Focus(ElementB);
+						ElementA->Click();
 
-					TEST_FALSE(Sequence->Perform());
-					TEST_FALSE(ElementB->IsFocused());
+						FDriverSequenceRef Sequence = Driver->CreateSequence();
+						Sequence->Actions()
+							.Wait(Until::ElementIsInteractable(ElementA, FWaitInterval::InSeconds(0.25), FWaitTimeout::InSeconds(1)))
+							.Focus(ElementB);
+
+						TEST_FALSE(Sequence->Perform());
+						TEST_FALSE(ElementB->IsFocused());
+					});
+
+					AfterEach([this]() {
+						SuiteViewModel->SetKeyResetDelay(FTimespan::FromSeconds(0.0));
+					});
 				});
 			});
 		});
