AutomationDriver - Fix sequence test on Linux

#rnx

diff --git a/Engine/Plugins/Tests/AutomationDriverTests/Source/AutomationDriverTests/Private/AutomationDriver.spec.cpp b/Engine/Plugins/Tests/AutomationDriverTests/Source/AutomationDriverTests/Private/AutomationDriver.spec.cpp
--- a/Engine/Plugins/Tests/AutomationDriverTests/Source/AutomationDriverTests/Private/AutomationDriver.spec.cpp
+++ b/Engine/Plugins/Tests/AutomationDriverTests/Source/AutomationDriverTests/Private/AutomationDriver.spec.cpp
@@ -1067,20 +1067,15 @@ void FAutomationDriverSpec::Define()
 		It("should be able to perform actions across multiple windows", EAsyncExecution::ThreadPool, [this]()
 		{
 			FDriverSequenceRef Sequence = Driver->CreateSequence();
+
 			Sequence->Actions()
 				.Click(By::Id("KeyModifierA#"))
 				.Click(By::Id("KeyA#"))
-				// Workaround for Mac. The first click closes previously opened popup menu
-				// and the second click opens needed popup menu
-				.Click(By::Id("KeyModifierEb"))
 				.Click(By::Id("KeyModifierEb"))
 				.Click(By::Id("KeyEb"))
-				// Workaround for Mac. The first click closes previously opened popup menu
-				// and the second click opens needed popup menu
-				.Click(By::Id("KeyModifierB#"))
 				.Click(By::Id("KeyModifierB#"))
 				.Click(By::Id("KeyB#"));
-
+			
 			TEST_TRUE(Sequence->Perform());
 			TEST_EQUAL(SuiteViewModel->GetKeySequence(), TEXT("A#EbB#"));
 		});
diff --git a/Engine/Plugins/Tests/AutomationDriverTests/Source/AutomationDriverTests/Private/SAutomationDriverSpecSuite.cpp b/Engine/Plugins/Tests/AutomationDriverTests/Source/AutomationDriverTests/Private/SAutomationDriverSpecSuite.cpp
--- a/Engine/Plugins/Tests/AutomationDriverTests/Source/AutomationDriverTests/Private/SAutomationDriverSpecSuite.cpp
+++ b/Engine/Plugins/Tests/AutomationDriverTests/Source/AutomationDriverTests/Private/SAutomationDriverSpecSuite.cpp
@@ -274,33 +274,32 @@ public:
 		return SNew(SBox)
 			[
 				SNew(SHorizontalBox)
-				+ SHorizontalBox::Slot()
-				[
-					ConstructPianoKey(KeyFlat)
-				]
-				+ SHorizontalBox::Slot()
-				[
-					ConstructPianoKey(KeySharp)
-				]
+					+ SHorizontalBox::Slot()
+					[
+						ConstructPianoContextMenuKey(KeyFlat)
+					]
+					+ SHorizontalBox::Slot()
+					[
+						ConstructPianoContextMenuKey(KeySharp)
+					]
 			];
 	}
 
+	TSharedRef<SWidget> ConstructPianoContextMenuKey(EPianoKey Key)
+	{
+		return ConstructPianoKeyImpl(SharedThis(this), &SAutomationDriverSpecSuiteImpl::KeyClicked, Key);
+	}
+
 	TSharedRef<SWidget> ConstructPianoKey(EPianoKey Key)
 	{
-		TSharedRef<SButton> Button = SNew(SButton)
-			.Text(FPianoKeyExtensions::ToText(Key))
-			.OnClicked(ViewModel.ToSharedRef(), &IAutomationDriverSpecSuiteViewModel::KeyClicked, Key)
-			.IsEnabled(ViewModel.ToSharedRef(), &IAutomationDriverSpecSuiteViewModel::IsKeyEnabled, Key)
-			.OnHovered(ViewModel.ToSharedRef(), &IAutomationDriverSpecSuiteViewModel::KeyHovered, Key)
-			.VAlign(VAlign_Bottom)
-			.HAlign(HAlign_Center)
-			.ContentPadding(FMargin(5, 0, 5, 10))
-			.Tag("Key")
-			.AddMetaData(FDriverMetaData::Id(*(FString(TEXT("Key")) + FPianoKeyExtensions::ToString(Key))));
+		return ConstructPianoKeyImpl(ViewModel.ToSharedRef(), &IAutomationDriverSpecSuiteViewModel::KeyClicked, Key);
+	}
 
-		PianoKeys.Add(Key, Button);
+	FReply KeyClicked(EPianoKey Key)
+	{
+		CloseContextMenu(Key);
 
-		return Button;
+		return ViewModel->KeyClicked(Key);
 	}
 
 	TSharedRef<SWidget> ConstructPianoKeyModifier(EPianoKey Key1, EPianoKey Key2)
@@ -345,6 +344,18 @@ public:
 		return FReply::Handled();
 	}
 
+	virtual FReply CloseContextMenu(EPianoKey KeySharp)
+	{
+		const TSharedPtr<SMenuAnchor>* Anchor = PianoKeyMenus.Find(KeySharp);
+
+		if (Anchor != nullptr)
+		{
+			(*Anchor)->SetIsOpen(false);
+		}
+
+		return FReply::Handled();
+	}
+
 	virtual TSharedPtr<SWidget> GetKeyWidget(EPianoKey Key) const override
 	{
 		const TSharedPtr<SButton>* Button = PianoKeys.Find(Key);
@@ -390,6 +401,27 @@ public:
 		DocumentScrollBox->ScrollToEnd();
 	}
 
+private:
+	template<typename UserClass>
+	TSharedRef<SWidget> ConstructPianoKeyImpl(TSharedRef<UserClass> InUserObjectPtr, FOnClicked::TMethodPtr<UserClass, EPianoKey> InKeyClickHandler, EPianoKey Key)
+	{
+		TSharedRef<SButton> Button = SNew(SButton)
+			.Text(FPianoKeyExtensions::ToText(Key))
+			.OnClicked(InUserObjectPtr, InKeyClickHandler, Key)
+			.IsEnabled(ViewModel.ToSharedRef(), &IAutomationDriverSpecSuiteViewModel::IsKeyEnabled, Key)
+			.OnHovered(ViewModel.ToSharedRef(), &IAutomationDriverSpecSuiteViewModel::KeyHovered, Key)
+			.VAlign(VAlign_Bottom)
+			.HAlign(HAlign_Center)
+			.ContentPadding(FMargin(5, 0, 5, 10))
+			.Tag("Key")
+			.AddMetaData(FDriverMetaData::Id(*(FString(TEXT("Key")) + FPianoKeyExtensions::ToString(Key))));
+
+		PianoKeys.Add(Key, Button);
+
+		return Button;
+	}
+
+
 private:
 
 	TSharedPtr<IAutomationDriverSpecSuiteViewModel> ViewModel;
