Add By::TextFilter::Contains and By::TextFilter::Equals locators to AutomationDriver

diff --git a/Engine/Plugins/Tests/AutomationDriverTests/Source/AutomationDriverTests/Private/AutomationDriver.spec.cpp b/Engine/Plugins/Tests/AutomationDriverTests/Source/AutomationDriverTests/Private/AutomationDriver.spec.cpp
--- a/Engine/Plugins/Tests/AutomationDriverTests/Source/AutomationDriverTests/Private/AutomationDriver.spec.cpp
+++ b/Engine/Plugins/Tests/AutomationDriverTests/Source/AutomationDriverTests/Private/AutomationDriver.spec.cpp
@@ -269,6 +269,142 @@ void FAutomationDriverSpec::Define()
 		});
 	});
 
+	Describe("By::TextFilter", [this]()
+	{
+		Describe("Equals", [this]()
+		{
+			static const FString ElementText = "Bb/A#";
+			static const FString ElementSubstringText = ElementText.LeftChop(1);
+			static const FString ElementWrongCaseText = ElementText.ToLower();
+
+			It("should fail to locate an element if none exist with the specified text", EAsyncExecution::ThreadPool, [this]()
+			{
+				TEST_FALSE(Driver->FindElement(By::TextFilter::Equals(By::Path("#Piano//<SButton>"), "NotDefined"))->Exists());
+			});
+
+			It("should fail to locate an element containing the specified substring in the text", EAsyncExecution::ThreadPool, [this]()
+			{
+				TEST_FALSE(Driver->FindElement(By::TextFilter::Equals(By::Path("#Piano//<SButton>"), ElementSubstringText))->Exists());
+			});
+
+			Describe("IgnoreCase", [this]()
+			{
+				It("should locate an element with the specified text in a different case", EAsyncExecution::ThreadPool, [this]()
+				{
+					TEST_TRUE(Driver->FindElement(By::TextFilter::Equals(By::Path("#Piano//<SButton>"), ElementWrongCaseText, ESearchCase::IgnoreCase))->Exists());
+				});
+
+				It("should locate an element with the specified text in the proper case", EAsyncExecution::ThreadPool, [this]()
+				{
+					TEST_TRUE(Driver->FindElement(By::TextFilter::Equals(By::Path("#Piano//<SButton>"), ElementText, ESearchCase::IgnoreCase))->Exists());
+				});
+			});
+
+			Describe("CaseSensitive", [this]()
+			{
+				It("should fail to locate an element with the specified text in a different case", EAsyncExecution::ThreadPool, [this]()
+				{
+					TEST_FALSE(Driver->FindElement(By::TextFilter::Equals(By::Path("#Piano//<SButton>"), ElementWrongCaseText, ESearchCase::CaseSensitive))->Exists());
+				});
+
+				It("should locate an element with the specified text in the proper case", EAsyncExecution::ThreadPool, [this]()
+				{
+					TEST_TRUE(Driver->FindElement(By::TextFilter::Equals(By::Path("#Piano//<SButton>"), ElementText, ESearchCase::CaseSensitive))->Exists());
+				});
+			});
+
+			Describe("Default", [this]()
+			{
+				It("should use case sensitive comparison by default", EAsyncExecution::ThreadPool, [this]()
+				{
+					TEST_TRUE(Driver->FindElement(By::TextFilter::Equals(By::Path("#Piano//<SButton>"), ElementText))->Exists());
+					TEST_FALSE(Driver->FindElement(By::TextFilter::Equals(By::Path("#Piano//<SButton>"), ElementWrongCaseText))->Exists());
+				});
+			});
+			
+			It("should be capable of locating multiple elements", EAsyncExecution::ThreadPool, [this]()
+			{
+				TEST_EQUAL(Driver->FindElements(By::TextFilter::Equals(By::Path("Documents//<SButton>"), "Document 1", ESearchCase::CaseSensitive))->GetElements().Num(), 3);
+			});
+		});
+		
+		Describe("Contains", [this]()
+		{
+			static const FString ElementText = "Bb/A#";
+			static const FString ElementSubstringText = ElementText.LeftChop(1);
+			static const FString ElementWrongCaseSubstringText = ElementSubstringText.ToLower();
+
+			static const FString ElementLeftSubstringText = ElementText.LeftChop(1);
+			static const FString ElementRightSubstringText = ElementText.RightChop(1);
+			static const FString ElementCenterSubstringText = ElementText.LeftChop(1).RightChop(1);
+
+			It("should fail to locate an element if none exist with the specified text", EAsyncExecution::ThreadPool, [this]()
+			{
+				TEST_FALSE(Driver->FindElement(By::TextFilter::Contains(By::Path("#Piano//<SButton>"), "NotDefined"))->Exists());
+			});
+
+			It("should locate an element containing the specified substring at any position in the text", EAsyncExecution::ThreadPool, [this]()
+			{
+				TEST_TRUE(Driver->FindElement(By::TextFilter::Contains(By::Path("#Piano//<SButton>"), ElementLeftSubstringText))->Exists());
+				TEST_TRUE(Driver->FindElement(By::TextFilter::Contains(By::Path("#Piano//<SButton>"), ElementCenterSubstringText))->Exists());
+				TEST_TRUE(Driver->FindElement(By::TextFilter::Contains(By::Path("#Piano//<SButton>"), ElementRightSubstringText))->Exists());
+			});
+
+			Describe("IgnoreCase", [this]()
+			{
+				It("should locate an element containing the specified substring in a different case", EAsyncExecution::ThreadPool, [this]()
+				{
+					TEST_TRUE(Driver->FindElement(By::TextFilter::Contains(By::Path("#Piano//<SButton>"), ElementWrongCaseSubstringText, ESearchCase::IgnoreCase))->Exists());
+				});
+
+				It("should locate an element containing the specified substring in the proper case", EAsyncExecution::ThreadPool, [this]()
+				{
+					TEST_TRUE(Driver->FindElement(By::TextFilter::Contains(By::Path("#Piano//<SButton>"), ElementSubstringText, ESearchCase::IgnoreCase))->Exists());
+				});
+			});
+
+			Describe("CaseSensitive", [this]()
+			{
+				It("should fail to locate an element containing the specified substring in a different case", EAsyncExecution::ThreadPool, [this]()
+				{
+					TEST_FALSE(Driver->FindElement(By::TextFilter::Contains(By::Path("#Piano//<SButton>"), ElementWrongCaseSubstringText, ESearchCase::CaseSensitive))->Exists());
+				});
+
+				It("should locate an element containing the specified substring in the proper case", EAsyncExecution::ThreadPool, [this]()
+				{
+					TEST_TRUE(Driver->FindElement(By::TextFilter::Contains(By::Path("#Piano//<SButton>"), ElementSubstringText, ESearchCase::CaseSensitive))->Exists());
+				});
+			});
+
+			Describe("SearchDir", [this]()
+			{
+				It("should locate an element containing the specified substring by searching in both directions", EAsyncExecution::ThreadPool, [this]()
+				{
+					TEST_TRUE(Driver->FindElement(By::TextFilter::Contains(By::Path("#Piano//<SButton>"), ElementLeftSubstringText, ESearchCase::CaseSensitive, ESearchDir::FromStart))->Exists());
+					TEST_TRUE(Driver->FindElement(By::TextFilter::Contains(By::Path("#Piano//<SButton>"), ElementLeftSubstringText, ESearchCase::CaseSensitive, ESearchDir::FromEnd))->Exists());
+					TEST_TRUE(Driver->FindElement(By::TextFilter::Contains(By::Path("#Piano//<SButton>"), ElementRightSubstringText, ESearchCase::CaseSensitive, ESearchDir::FromStart))->Exists());
+					TEST_TRUE(Driver->FindElement(By::TextFilter::Contains(By::Path("#Piano//<SButton>"), ElementRightSubstringText, ESearchCase::CaseSensitive, ESearchDir::FromEnd))->Exists());
+				});
+			});
+
+			Describe("Default", [this]()
+			{
+				It("should ignore case by default", EAsyncExecution::ThreadPool, [this]()
+				{
+					TEST_TRUE(Driver->FindElement(By::TextFilter::Contains(By::Path("#Piano//<SButton>"), ElementSubstringText))->Exists());
+					TEST_TRUE(Driver->FindElement(By::TextFilter::Contains(By::Path("#Piano//<SButton>"), ElementWrongCaseSubstringText))->Exists());
+				});
+			});
+			
+			It("should be capable of locating multiple elements", EAsyncExecution::ThreadPool, [this]()
+			{
+				const FString ElementsSubstringText = "ocument 2";
+				const int32 ElementsNum = 12; // "Document 2", "Document 20", ..., "Document 29", "Document 200"
+				TEST_EQUAL(Driver->FindElements(By::TextFilter::Contains(By::Path("#Documents//<SScrollBox>//<SButton>"), ElementsSubstringText, ESearchCase::CaseSensitive))->GetElements().Num(), ElementsNum);
+			});
+		});
+	});
+
 	Describe("Element", [this]()
 	{
 		Describe("Hover", [this]()
diff --git a/Engine/Source/Developer/AutomationDriver/Private/LocateBy.cpp b/Engine/Source/Developer/AutomationDriver/Private/LocateBy.cpp
--- a/Engine/Source/Developer/AutomationDriver/Private/LocateBy.cpp
+++ b/Engine/Source/Developer/AutomationDriver/Private/LocateBy.cpp
@@ -3,8 +3,10 @@
 #include "LocateBy.h"
 #include "Locators/SlateWidgetLocatorByDelegate.h"
 #include "Locators/SlateWidgetLocatorByPath.h"
+#include "Locators/WidgetLocatorByFilter.h"
 #include "AutomationDriverTypeDefs.h"
 #include "Framework/Application/SlateApplication.h"
+#include "IApplicationElement.h"
 
 TSharedRef<IElementLocator, ESPMode::ThreadSafe> By::Delegate(const FLocateSlateWidgetElementDelegate& Value)
 {
@@ -119,3 +121,34 @@ TSharedRef<IElementLocator, ESPMode::ThreadSafe> By::Cursor()
 		}
 	});
 }
+
+TSharedRef<IElementLocator, ESPMode::ThreadSafe> By::TextFilter::Contains(
+	const FElementLocatorRef& RootLocator,
+	const FString& Value,
+	ESearchCase::Type SearchCase,
+	ESearchDir::Type SearchDir)
+{
+	return FWidgetLocatorByFilterFactory::Create(
+		TEXT("[By::TextFilter::Contains] ") + Value,
+		RootLocator,
+		[Value, SearchCase, SearchDir](const TSharedRef<IApplicationElement>& Element)
+		{
+			return Element->GetText().ToString().Contains(Value, SearchCase, SearchDir);
+		}
+	);
+}
+
+TSharedRef<IElementLocator, ESPMode::ThreadSafe> By::TextFilter::Equals(
+	const FElementLocatorRef& RootLocator,
+	const FString& Value,
+	ESearchCase::Type SearchCase)
+{
+	return FWidgetLocatorByFilterFactory::Create(
+		TEXT("[By::TextFilter::Equals] ") + Value,
+		RootLocator,
+		[Value, SearchCase](const TSharedRef<IApplicationElement>& Element)
+		{
+			return Element->GetText().ToString().Equals(Value, SearchCase);
+		}
+	);
+}
diff --git a/Engine/Source/Developer/AutomationDriver/Private/Locators/WidgetLocatorByFilter.cpp b/Engine/Source/Developer/AutomationDriver/Private/Locators/WidgetLocatorByFilter.cpp
new file mode 100644
--- /dev/null
+++ b/Engine/Source/Developer/AutomationDriver/Private/Locators/WidgetLocatorByFilter.cpp
@@ -0,0 +1,62 @@
+// Copyright Epic Games, Inc. All Rights Reserved.
+
+#include "Locators/WidgetLocatorByFilter.h"
+#include "IElementLocator.h"
+
+class FWidgetLocatorByFilter
+	: public IElementLocator
+{
+public:
+	using FFilterFunction = FWidgetLocatorByFilterFactory::FFilterFunction;
+
+	virtual ~FWidgetLocatorByFilter()
+	{ }
+
+	virtual FString ToDebugString() const
+	{
+		return RootLocator->ToDebugString() + " " + DebugString;
+	}
+
+	virtual void Locate(TArray<TSharedRef<IApplicationElement>>& OutElements) const override
+	{
+		RootLocator->Locate(OutElements);
+
+		for (int32 i = 0; i < OutElements.Num();)
+		{
+			if (Filter(OutElements[i]))
+			{
+				i++;
+			}
+			else
+			{
+				OutElements.RemoveAtSwap(i);
+			}
+		}
+	}
+
+private:
+	FWidgetLocatorByFilter(
+		const FString& InDebugString,
+		const FElementLocatorRef& InRootLocator,
+		FFilterFunction InFilter) :
+		DebugString{ InDebugString },
+		RootLocator{ InRootLocator },
+		Filter{ MoveTemp(InFilter) }
+	{ }
+
+private:
+
+	const FString DebugString;
+	const FElementLocatorRef RootLocator;
+	const FFilterFunction Filter;
+
+	friend FWidgetLocatorByFilterFactory;
+};
+
+TSharedRef<IElementLocator, ESPMode::ThreadSafe> FWidgetLocatorByFilterFactory::Create(
+	const FString& DebugString,
+	const FElementLocatorRef& RootLocator,
+	FFilterFunction Filter)
+{
+	return MakeShareable(new FWidgetLocatorByFilter(DebugString, RootLocator, MoveTemp(Filter)));
+}
diff --git a/Engine/Source/Developer/AutomationDriver/Private/Locators/WidgetLocatorByFilter.h b/Engine/Source/Developer/AutomationDriver/Private/Locators/WidgetLocatorByFilter.h
new file mode 100644
--- /dev/null
+++ b/Engine/Source/Developer/AutomationDriver/Private/Locators/WidgetLocatorByFilter.h
@@ -0,0 +1,29 @@
+// Copyright Epic Games, Inc. All Rights Reserved.
+#pragma once
+
+#include "CoreMinimal.h"
+#include "AutomationDriverTypeDefs.h"
+
+class IApplicationElement;
+class IElementLocator;
+
+/* This class provides a utility function for creating instances of the filter locator. */
+class FWidgetLocatorByFilterFactory
+{
+public:
+	using FFilterFunction = TFunction<bool(const TSharedRef<IApplicationElement>&)>;
+	
+	/**
+	* Creates a new element locator that limits the elements discovered by the provided root locator to those determined by the filter function.
+	*
+	* @param DebugString - The string used for debugging and logging purposes to represent this particular instance of the filter locator
+	* @param RootLocator - The reference to the element locator where the search will be started from
+	* @param Filter - The function applied to the elements found by the root locator to determine which should be saved as a result.
+	* @return a locator which uses the specified root locator and a filter function to discover appropriate elements
+	*/
+	static TSharedRef<IElementLocator, ESPMode::ThreadSafe> Create(
+		const FString& DebugString,
+		const FElementLocatorRef& RootLocator,
+		FFilterFunction Filter
+	);
+};
diff --git a/Engine/Source/Developer/AutomationDriver/Public/LocateBy.h b/Engine/Source/Developer/AutomationDriver/Public/LocateBy.h
--- a/Engine/Source/Developer/AutomationDriver/Public/LocateBy.h
+++ b/Engine/Source/Developer/AutomationDriver/Public/LocateBy.h
@@ -313,4 +313,41 @@ public:
 	 * @return a locator which uses the cursor position to discover appropriate elements
 	 */
 	static TSharedRef<IElementLocator, ESPMode::ThreadSafe> Cursor();
+	
+	class AUTOMATIONDRIVER_API TextFilter
+	{
+	public:
+
+		/**
+		 * Creates a new element locator that limits the elements discovered by the root locator provided to the function
+		 * to those containing the specified text substring.
+		 *
+		 * @param RootLocator - The reference to the element locator where the search will be started from
+		 * @param Value - The text of the element to find
+		 * @param SearchCase - Indicates whether the search is case sensitive or not ( defaults to ESearchCase::IgnoreCase )
+		 * @param SearchDir - Indicates whether the search starts at the beginning or at the end ( defaults to ESearchDir::FromStart )
+		 * @return a locator which uses the specified root locator and a text string to discover appropriate elements
+		 */
+		static TSharedRef<IElementLocator, ESPMode::ThreadSafe> Contains(
+			const FElementLocatorRef& RootLocator,
+			const FString& Value,
+			ESearchCase::Type SearchCase = ESearchCase::IgnoreCase,
+			ESearchDir::Type SearchDir = ESearchDir::FromStart
+		);
+
+		/**
+		 * Creates a new element locator that limits the elements discovered by the root locator provided to the function
+		 * to those containing the specified text.
+		 *
+		 * @param RootLocator - The reference to the element locator where the search will be started from
+		 * @param Value - The text of the element to find
+		 * @param SearchCase - Indicates whether the search is case sensitive or not ( defaults to ESearchCase::CaseSensitive )
+		 * @return a locator which uses the specified root locator and a text string to discover appropriate elements
+		 */
+		static TSharedRef<IElementLocator, ESPMode::ThreadSafe> Equals(
+			const FElementLocatorRef& RootLocator,
+			const FString& Value,
+			ESearchCase::Type SearchCase = ESearchCase::CaseSensitive
+		);
+	};
 };
