Fix scrolling of elements with duplicated ScrollBar in AutomationDriver.

The fix is required for the automated test case that uses the AutomationDriver to scroll the SDetailsView widget.

#rnx

diff --git a/Engine/Plugins/Tests/AutomationDriverTests/Source/AutomationDriverTests/Private/SAutomationDriverSpecSuite.cpp b/Engine/Plugins/Tests/AutomationDriverTests/Source/AutomationDriverTests/Private/SAutomationDriverSpecSuite.cpp
--- a/Engine/Plugins/Tests/AutomationDriverTests/Source/AutomationDriverTests/Private/SAutomationDriverSpecSuite.cpp
+++ b/Engine/Plugins/Tests/AutomationDriverTests/Source/AutomationDriverTests/Private/SAutomationDriverSpecSuite.cpp
@@ -24,6 +24,8 @@ public:
 	{
 		ViewModel = InViewModel;
 		AddMetadata(FDriverMetaData::Id("Suite"));
+		
+		TSharedRef<SScrollBar> TilesScrollbar = SNew(SScrollBar);
 
 		SAssignNew(WindowContents, SVerticalBox)
 		+ SVerticalBox::Slot()
@@ -150,13 +152,36 @@ public:
 			+ SHorizontalBox::Slot()
 			.Padding(5, 5)
 			[
-				SAssignNew(DocumentTiles, STileView<TSharedRef<FDocumentInfo>>)
-				.ItemHeight(48)
-				.ItemWidth(80)
-				.SelectionMode(ESelectionMode::None)
-				.ListItemsSource(&InViewModel->GetDocuments())
-				.OnGenerateTile(this, &SAutomationDriverSpecSuiteImpl::GenerateListRow)
+				SNew(SHorizontalBox)
 				.Tag("Tiles")
+				+ SHorizontalBox::Slot()
+				.AutoWidth()
+				[
+					SNew(SBox)
+					.WidthOverride(16.0f)
+					[
+						TilesScrollbar
+					]
+				]
+				+ SHorizontalBox::Slot()
+				[
+					SAssignNew(DocumentTiles, STileView<TSharedRef<FDocumentInfo>>)
+					.ItemHeight(48)
+					.ItemWidth(80)
+					.SelectionMode(ESelectionMode::None)
+					.ListItemsSource(&InViewModel->GetDocuments())
+					.OnGenerateTile(this, &SAutomationDriverSpecSuiteImpl::GenerateListRow)
+					.ExternalScrollbar(TilesScrollbar)
+				]
+				+ SHorizontalBox::Slot()
+				.AutoWidth()
+				[
+					SNew(SBox)
+					.WidthOverride(16.0f)
+					[
+						TilesScrollbar
+					]
+				]
 			]
 			+ SHorizontalBox::Slot()
 			.Padding(5, 5)
diff --git a/Engine/Source/Developer/AutomationDriver/Private/SlateWidgetElement.cpp b/Engine/Source/Developer/AutomationDriver/Private/SlateWidgetElement.cpp
--- a/Engine/Source/Developer/AutomationDriver/Private/SlateWidgetElement.cpp
+++ b/Engine/Source/Developer/AutomationDriver/Private/SlateWidgetElement.cpp
@@ -399,7 +399,7 @@ public:
 				
 				if (FoundScrollBar->IsNeeded())
 				{
-					if (ScrollBar.IsValid())
+					if (ScrollBar.IsValid() && ScrollBar != FoundScrollBar)
 					{
 						ErrorCode = 2;
 						return nullptr;
